EnhancedInput: Added ability to bind actions with variable arguments;

diff --git a/Engine/Plugins/Experimental/EnhancedInput/Source/EnhancedInput/Public/EnhancedInputComponent.h b/Engine/Plugins/Experimental/EnhancedInput/Source/EnhancedInput/Public/EnhancedInputComponent.h
--- a/Engine/Plugins/Experimental/EnhancedInput/Source/EnhancedInput/Public/EnhancedInputComponent.h
+++ b/Engine/Plugins/Experimental/EnhancedInput/Source/EnhancedInput/Public/EnhancedInputComponent.h
@@ -72,6 +72,14 @@ public:
 		Delegate->BindUFunction(Object, FuncName);
 	}
 
+	template<typename TSig = TSignature>
+	TSig& MakeDelegate()
+	{
+		Unbind();
+		Delegate = MakeShared<TSig>();
+		return *Delegate;
+	}
+
 	template<typename... TArgs>
 	void Execute(TArgs... Args) const
 	{
@@ -330,6 +338,14 @@ public:
 	DEFINE_BIND_ACTION(FEnhancedInputActionHandlerValueSignature);
 	DEFINE_BIND_ACTION(FEnhancedInputActionHandlerInstanceSignature);
 
+	template< class FuncType, class UserClass, typename... VarTypes >
+	FEnhancedInputActionEventBinding& BindAction(const UInputAction* Action, ETriggerEvent TriggerEvent, UserClass* Object, FuncType Func, VarTypes... Vars)
+	{
+		TUniquePtr<FEnhancedInputActionEventDelegateBinding<FEnhancedInputActionHandlerSignature>> AB = MakeUnique<FEnhancedInputActionEventDelegateBinding<FEnhancedInputActionHandlerSignature>>(Action, TriggerEvent);
+		AB->Delegate.MakeDelegate().BindUObject(Object, Func, Vars...);
+		return *EnhancedActionEventBindings.Add_GetRef(MoveTemp(AB));
+	}
+
 	/**
 	 * Binds to an object UFUNCTION
 	 */
