EnhancedInput: Remove support for Input.ini +AxisConfig settings for all devices except mice. This stops the default config settings applying a hidden +-0.25 deadzone to any gamepad axis.

diff --git a/Engine/Plugins/Experimental/EnhancedInput/Source/EnhancedInput/Private/EnhancedInputSubsystemInterface.cpp b/Engine/Plugins/Experimental/EnhancedInput/Source/EnhancedInput/Private/EnhancedInputSubsystemInterface.cpp
--- a/Engine/Plugins/Experimental/EnhancedInput/Source/EnhancedInput/Private/EnhancedInputSubsystemInterface.cpp
+++ b/Engine/Plugins/Experimental/EnhancedInput/Source/EnhancedInput/Private/EnhancedInputSubsystemInterface.cpp
@@ -234,6 +234,12 @@ void IEnhancedInputSubsystemInterface::ApplyAxisPropertyModifiers(UEnhancedPlaye
 	//	return;
 	//}
 
+	// TODO: This function is causing issues with gamepads, applying a hidden 0.25 deadzone modifier by default. Apply it to mouse inputs only until a better system is in place.
+	if (!Mapping.Key.IsMouseButton())
+	{
+		return;
+	}
+
 	// Apply applicable axis property modifiers from the old input settings automatically.
 	// TODO: This needs to live at the EnhancedInputSettings level.
 	// TODO: Adopt this approach for all modifiers? Most of these are better done at the action level for most use cases.
