EnhancedInput: Added method to query mapped keys for an InputAction.

diff --git a/Engine/Plugins/Experimental/EnhancedInput/Source/EnhancedInput/Private/EnhancedInputSubsystemInterface.cpp b/Engine/Plugins/Experimental/EnhancedInput/Source/EnhancedInput/Private/EnhancedInputSubsystemInterface.cpp
--- a/Engine/Plugins/Experimental/EnhancedInput/Source/EnhancedInput/Private/EnhancedInputSubsystemInterface.cpp
+++ b/Engine/Plugins/Experimental/EnhancedInput/Source/EnhancedInput/Private/EnhancedInputSubsystemInterface.cpp
@@ -303,6 +303,24 @@ bool IEnhancedInputSubsystemInterface::HasMappingContext(const UInputMappingCont
 	return GetPlayerInput() && GetPlayerInput()->AppliedInputContexts.Contains(MappingContext);
 }
 
+TArray<FKey> IEnhancedInputSubsystemInterface::QueryKeysMappedToAction(const UInputAction* Action) const
+{
+	TArray<FKey> MappedKeys;
+
+	if (const UEnhancedPlayerInput* const PlayerInput = GetPlayerInput())
+	{
+		for (const FEnhancedActionKeyMapping& Mapping : PlayerInput->EnhancedActionMappings)
+		{
+			if (Mapping.Action == Action)
+			{
+				MappedKeys.AddUnique(Mapping.Key);
+			}
+		}
+	}
+
+	return MappedKeys;
+}
+
 template<typename T>
 void DeepCopyPtrArray(const TArray<T*>& From, TArray<T*>& To)
 {
diff --git a/Engine/Plugins/Experimental/EnhancedInput/Source/EnhancedInput/Public/EnhancedInputSubsystemInterface.h b/Engine/Plugins/Experimental/EnhancedInput/Source/EnhancedInput/Public/EnhancedInputSubsystemInterface.h
--- a/Engine/Plugins/Experimental/EnhancedInput/Source/EnhancedInput/Public/EnhancedInputSubsystemInterface.h
+++ b/Engine/Plugins/Experimental/EnhancedInput/Source/EnhancedInput/Public/EnhancedInputSubsystemInterface.h
@@ -94,6 +94,12 @@ public:
 	UFUNCTION(BlueprintCallable, BlueprintCosmetic, Category = "Input|Mapping Queries")	// TODO: BlueprintPure would be nicer. Move into library?
 	virtual bool HasMappingContext(const UInputMappingContext* MappingContext) const;
 
+	/**
+	 * Returns the keys mapped to the given action in the active input mapping contexts.
+	 */
+	UFUNCTION(BlueprintCallable, BlueprintCosmetic, Category = "Input|Mapping Queries")
+	virtual TArray<FKey> QueryKeysMappedToAction(const UInputAction* Action) const;
+
 private:
 
 	// Forced actions/keys for debug. These will be applied each tick once set even if zeroed, until removed. 
