Add a CVar to allow Global Axis Config for the Enhanced input Subsystem.

diff --git a/Engine/Plugins/Experimental/EnhancedInput/Source/EnhancedInput/Private/EnhancedInputSubsystemInterface.cpp b/Engine/Plugins/Experimental/EnhancedInput/Source/EnhancedInput/Private/EnhancedInputSubsystemInterface.cpp
--- a/Engine/Plugins/Experimental/EnhancedInput/Source/EnhancedInput/Private/EnhancedInputSubsystemInterface.cpp
+++ b/Engine/Plugins/Experimental/EnhancedInput/Source/EnhancedInput/Private/EnhancedInputSubsystemInterface.cpp
@@ -16,6 +16,17 @@
  * See EnhancedInputSubsystemInterfaceDebug.cpp for debug specific functionality.
  */
 
+static constexpr int32 GGlobalAxisConfigMode_Default = 0;
+static constexpr int32 GGlobalAxisConfigMode_All = 1;
+static constexpr int32 GGlobalAxisConfigMode_None = 2;
+
+static int32 GGlobalAxisConfigMode = 0;
+static FAutoConsoleVariableRef GCVarGlobalAxisConfigMode(
+	TEXT("input.GlobalAxisConfigMode"),
+	GGlobalAxisConfigMode,
+	TEXT("Whether or not to apply Global Axis Config settings. 0 = Default (Mouse Only), 1 = All, 2 = None")
+);
+
 void IEnhancedInputSubsystemInterface::ClearAllMappings()
 {
 	if (UEnhancedPlayerInput* PlayerInput = GetPlayerInput())
@@ -234,8 +245,14 @@ void IEnhancedInputSubsystemInterface::ApplyAxisPropertyModifiers(UEnhancedPlaye
 	//	return;
 	//}
 
+	if (GGlobalAxisConfigMode_None == GGlobalAxisConfigMode)
+	{
+		return;
+	}
+
 	// TODO: This function is causing issues with gamepads, applying a hidden 0.25 deadzone modifier by default. Apply it to mouse inputs only until a better system is in place.
-	if (!Mapping.Key.IsMouseButton())
+	if (GGlobalAxisConfigMode_All != GGlobalAxisConfigMode &&
+		!Mapping.Key.IsMouseButton())
 	{
 		return;
 	}
