EnhancedInput: Retain appropriate action and  mapping trigger/modifier state when rebuilding mappings.

diff --git a/Engine/Plugins/Experimental/EnhancedInput/Source/EnhancedInput/Private/EnhancedInputSubsystemInterface.cpp b/Engine/Plugins/Experimental/EnhancedInput/Source/EnhancedInput/Private/EnhancedInputSubsystemInterface.cpp
--- a/Engine/Plugins/Experimental/EnhancedInput/Source/EnhancedInput/Private/EnhancedInputSubsystemInterface.cpp
+++ b/Engine/Plugins/Experimental/EnhancedInput/Source/EnhancedInput/Private/EnhancedInputSubsystemInterface.cpp
@@ -364,8 +364,9 @@ void IEnhancedInputSubsystemInterface::RebuildControlMappings()
 		return;
 	}
 
+	// Clear existing mappings, but retain the mapping array for later processing
+	TArray<FEnhancedActionKeyMapping> OldMappings(MoveTemp(PlayerInput->EnhancedActionMappings));
 	PlayerInput->ClearAllMappings();
-	PlayerInput->ResetActionInstanceData();
 
 	// Order contexts by priority
 	TMap<const UInputMappingContext*, int32> OrderedInputContexts = PlayerInput->AppliedInputContexts;
@@ -437,6 +438,29 @@ void IEnhancedInputSubsystemInterface::RebuildControlMappings()
 
 	PlayerInput->ForceRebuildingKeyMaps(false);
 
+	// Remove action instance data for actions that are not referenced in the new action mappings
+	TSet<const UInputAction*> RemovedActions;
+	for (TPair<const UInputAction*, FInputActionInstance>& ActionInstance : PlayerInput->ActionInstanceData)
+	{
+		RemovedActions.Add(ActionInstance.Key);
+	}
+	for (FEnhancedActionKeyMapping& Mapping : PlayerInput->EnhancedActionMappings)
+	{
+		RemovedActions.Remove(Mapping.Action);
+
+		// Retain old mapping trigger/modifier state for identical key -> action mappings.
+		TArray<FEnhancedActionKeyMapping>::SizeType Idx = OldMappings.IndexOfByPredicate([&Mapping](const FEnhancedActionKeyMapping& Other) {return Mapping.Action == Other.Action && Mapping.Key == Other.Key; });
+		if (Idx != INDEX_NONE)
+		{
+			Mapping = MoveTemp(OldMappings[Idx]);
+			OldMappings.RemoveAtSwap(Idx);
+		}
+	}
+	for (const UInputAction* Action : RemovedActions)
+	{
+		PlayerInput->ActionInstanceData.Remove(Action);
+	}
+
 	bMappingRebuildPending = false;
 }
 
diff --git a/Engine/Plugins/Experimental/EnhancedInput/Source/EnhancedInput/Public/EnhancedPlayerInput.h b/Engine/Plugins/Experimental/EnhancedInput/Source/EnhancedInput/Public/EnhancedPlayerInput.h
--- a/Engine/Plugins/Experimental/EnhancedInput/Source/EnhancedInput/Public/EnhancedPlayerInput.h
+++ b/Engine/Plugins/Experimental/EnhancedInput/Source/EnhancedInput/Public/EnhancedPlayerInput.h
@@ -67,11 +67,6 @@ private:
 
 	FInputActionInstance& FindOrAddActionEventData(const UInputAction* Action) const;
 
-	void ResetActionInstanceData()
-	{
-		ActionInstanceData.Reset();
-	}
-
 	template<typename T>
 	void GatherActionEventDataForActionMap(const T& ActionMap, TMap<const UInputAction*, FInputActionInstance>& FoundActionEventData) const;
 
