Add nullcheck to EnhancedInputMappingInterface when we are applying global axis mappings.
Prevents a crash that can happen when users have a Modifier on an InputAction, but don't specify a type.

diff --git a/Engine/Plugins/Experimental/EnhancedInput/Source/EnhancedInput/Private/EnhancedInputSubsystemInterface.cpp b/Engine/Plugins/Experimental/EnhancedInput/Source/EnhancedInput/Private/EnhancedInputSubsystemInterface.cpp
--- a/Engine/Plugins/Experimental/EnhancedInput/Source/EnhancedInput/Private/EnhancedInputSubsystemInterface.cpp
+++ b/Engine/Plugins/Experimental/EnhancedInput/Source/EnhancedInput/Private/EnhancedInputSubsystemInterface.cpp
@@ -268,7 +268,7 @@ void IEnhancedInputSubsystemInterface::ApplyAxisPropertyModifiers(UEnhancedPlaye
 		// If a modifier already exists it should override axis properties.
 		auto HasExistingModifier = [&Mapping](UClass* OfType)
 		{
-			auto TypeMatcher = [&OfType](UInputModifier* Modifier) { return Modifier->IsA(OfType); };
+			auto TypeMatcher = [&OfType](UInputModifier* Modifier) { return Modifier != nullptr && Modifier->IsA(OfType); };
 			return Mapping.Modifiers.ContainsByPredicate(TypeMatcher) || Mapping.Action->Modifiers.ContainsByPredicate(TypeMatcher);
 		};
 
