Renamed EMappingQueryIssue for use in Python

#rnx

diff --git a/Engine/Plugins/Experimental/EnhancedInput/Source/EnhancedInput/Public/InputMappingQuery.h b/Engine/Plugins/Experimental/EnhancedInput/Source/EnhancedInput/Public/InputMappingQuery.h
--- a/Engine/Plugins/Experimental/EnhancedInput/Source/EnhancedInput/Public/InputMappingQuery.h
+++ b/Engine/Plugins/Experimental/EnhancedInput/Source/EnhancedInput/Public/InputMappingQuery.h
@@ -24,7 +24,7 @@ enum class EMappingQueryResult : uint8
 };
 
 // Mapping issues arising from a QueryMapKeyIn... call
-UENUM(BlueprintType, meta = (Bitflags, UseEnumValuesAsMaskValuesInEditor = "true"))
+UENUM(BlueprintType, meta = (Bitflags, UseEnumValuesAsMaskValuesInEditor = "true", ScriptName = "MappingQueryIssueFlag"))
 enum class EMappingQueryIssue : uint8
 {
 	// Mapping will not affect any existing mappings and is safe to apply.
