EnhancedInput: Store chorded mappings by array index rather than pointer, avoiding invalidating pointers on array resize.

diff --git a/Engine/Plugins/Experimental/EnhancedInput/Source/EnhancedInput/Private/EnhancedInputSubsystemInterface.cpp b/Engine/Plugins/Experimental/EnhancedInput/Source/EnhancedInput/Private/EnhancedInputSubsystemInterface.cpp
--- a/Engine/Plugins/Experimental/EnhancedInput/Source/EnhancedInput/Private/EnhancedInputSubsystemInterface.cpp
+++ b/Engine/Plugins/Experimental/EnhancedInput/Source/EnhancedInput/Private/EnhancedInputSubsystemInterface.cpp
@@ -200,7 +200,7 @@ bool IEnhancedInputSubsystemInterface::HasTriggerWith(TFunctionRef<bool(const UI
 	return false;
 };
 
-void IEnhancedInputSubsystemInterface::InjectChordBlockers(const TMap<FEnhancedActionKeyMapping*, int32>& ChordedMappings)
+void IEnhancedInputSubsystemInterface::InjectChordBlockers(const TArray<int32>& ChordedMappings)
 {
 	UEnhancedPlayerInput* PlayerInput = GetPlayerInput();
 	if (!PlayerInput)
@@ -209,13 +209,13 @@ void IEnhancedInputSubsystemInterface::InjectChordBlockers(const TMap<FEnhancedA
 	}
 
 	// Inject chord blockers into all lower priority action mappings with a shared key
-	for (const TPair<FEnhancedActionKeyMapping*, int32>& ChordPair : ChordedMappings)
+	for (int32 MappingIndex : ChordedMappings)
 	{
-		FEnhancedActionKeyMapping* ChordMapping = ChordPair.Key;
-		for (int32 i = ChordPair.Value; i < PlayerInput->EnhancedActionMappings.Num(); ++i)
+		FEnhancedActionKeyMapping& ChordMapping = PlayerInput->EnhancedActionMappings[MappingIndex];
+		for (int32 i = MappingIndex + 1; i < PlayerInput->EnhancedActionMappings.Num(); ++i)
 		{
 			FEnhancedActionKeyMapping& Mapping = PlayerInput->EnhancedActionMappings[i];
-			if (Mapping.Action && Mapping.Key == ChordMapping->Key)
+			if (Mapping.Action && Mapping.Key == ChordMapping.Key)
 			{
 				// If we have no explicit triggers we can't inject an implicit as it may cause us to fire when we shouldn't.
 				auto AnyExplicit = [](const UInputTrigger* Trigger) { return Trigger->GetTriggerType() == ETriggerType::Explicit; };
@@ -227,7 +227,7 @@ void IEnhancedInputSubsystemInterface::InjectChordBlockers(const TMap<FEnhancedA
 				}
 
 				UInputTriggerChordBlocker* ChordBlocker = NewObject<UInputTriggerChordBlocker>(PlayerInput);
-				ChordBlocker->ChordAction = ChordMapping->Action;
+				ChordBlocker->ChordAction = ChordMapping.Action;
 				// TODO: If the chording action is bound at a lower priority than the blocked action its trigger state will be evaluated too late, which may produce unintended effects on the first tick.
 				Mapping.Triggers.Add(ChordBlocker);
 			}
@@ -374,7 +374,7 @@ void IEnhancedInputSubsystemInterface::RebuildControlMappings()
 
 	TSet<FKey> AppliedKeys;
 
-	TMap<FEnhancedActionKeyMapping*, int32> ChordedMappings;
+	TArray<int32> ChordedMappings;
 
 	for (const TPair<const UInputMappingContext*, int32>& ContextPair : OrderedInputContexts)
 	{
@@ -396,7 +396,8 @@ void IEnhancedInputSubsystemInterface::RebuildControlMappings()
 					ContextAppliedKeys.Add(Mapping.Key);
 				}
 
-				FEnhancedActionKeyMapping& NewMapping = PlayerInput->AddMapping(Mapping);
+				int32 NewMappingIndex = PlayerInput->AddMapping(Mapping);
+				FEnhancedActionKeyMapping& NewMapping = PlayerInput->EnhancedActionMappings[NewMappingIndex];
 
 				// Re-instance modifiers
 				DeepCopyPtrArray<UInputModifier>(Mapping.Modifiers, NewMapping.Modifiers);
@@ -413,7 +414,7 @@ void IEnhancedInputSubsystemInterface::RebuildControlMappings()
 				{
 					// TODO: Re-prioritize chorded mappings (within same context only?) by number of chorded actions, so Ctrl + Alt + [key] > Ctrl + [key] > [key].
 					// TODO: Above example shouldn't block [key] if only Alt is down, as there is no direct Alt + [key] mapping.y
-					ChordedMappings.Emplace(&NewMapping, PlayerInput->EnhancedActionMappings.Num());
+					ChordedMappings.Add(NewMappingIndex);
 
 					// Action level chording triggers need to be evaluated at the mapping level to ensure they block early enough.
 					// TODO: Continuing to evaluate these at the action level is redundant.
diff --git a/Engine/Plugins/Experimental/EnhancedInput/Source/EnhancedInput/Private/EnhancedPlayerInput.cpp b/Engine/Plugins/Experimental/EnhancedInput/Source/EnhancedInput/Private/EnhancedPlayerInput.cpp
--- a/Engine/Plugins/Experimental/EnhancedInput/Source/EnhancedInput/Private/EnhancedPlayerInput.cpp
+++ b/Engine/Plugins/Experimental/EnhancedInput/Source/EnhancedInput/Private/EnhancedPlayerInput.cpp
@@ -416,13 +416,13 @@ FInputActionValue UEnhancedPlayerInput::GetActionValue(const UInputAction* ForAc
 }
 
 
-FEnhancedActionKeyMapping& UEnhancedPlayerInput::AddMapping(const FEnhancedActionKeyMapping& Mapping)
+int32 UEnhancedPlayerInput::AddMapping(const FEnhancedActionKeyMapping& Mapping)
 {
 	int32 MappingIndex = EnhancedActionMappings.AddUnique(Mapping);
 	++EnhancedKeyBinds.FindOrAdd(Mapping.Key);
 	bKeyMapsBuilt = false;
 
-	return EnhancedActionMappings[MappingIndex];
+	return MappingIndex;
 }
 
 void UEnhancedPlayerInput::ClearAllMappings()
diff --git a/Engine/Plugins/Experimental/EnhancedInput/Source/EnhancedInput/Public/EnhancedInputSubsystemInterface.h b/Engine/Plugins/Experimental/EnhancedInput/Source/EnhancedInput/Public/EnhancedInputSubsystemInterface.h
--- a/Engine/Plugins/Experimental/EnhancedInput/Source/EnhancedInput/Public/EnhancedInputSubsystemInterface.h
+++ b/Engine/Plugins/Experimental/EnhancedInput/Source/EnhancedInput/Public/EnhancedInputSubsystemInterface.h
@@ -109,7 +109,7 @@ private:
 	void RemoveForcedInput(FKey Key);
 	void TickForcedInput(float DeltaTime);
 
-	void InjectChordBlockers(const TMap<struct FEnhancedActionKeyMapping*, int32>& ChordedMappings);
+	void InjectChordBlockers(const TArray<int32>& ChordedMappings);
 	bool HasTriggerWith(TFunctionRef<bool(const class UInputTrigger*)> TestFn, const TArray<class UInputTrigger*>& Triggers);
 
 	/** Reapply all control mappings to players pending a rebuild */
diff --git a/Engine/Plugins/Experimental/EnhancedInput/Source/EnhancedInput/Public/EnhancedPlayerInput.h b/Engine/Plugins/Experimental/EnhancedInput/Source/EnhancedInput/Public/EnhancedPlayerInput.h
--- a/Engine/Plugins/Experimental/EnhancedInput/Source/EnhancedInput/Public/EnhancedPlayerInput.h
+++ b/Engine/Plugins/Experimental/EnhancedInput/Source/EnhancedInput/Public/EnhancedPlayerInput.h
@@ -50,8 +50,10 @@ protected:
 
 private:
 
-	/** Add a player specific action mapping. */
-	FEnhancedActionKeyMapping& AddMapping(const FEnhancedActionKeyMapping& Mapping);
+	/** Add a player specific action mapping.
+	* Returns index into EnhancedActionMappings array.
+	*/
+	int32 AddMapping(const FEnhancedActionKeyMapping& Mapping);
 	void ClearAllMappings();
 
 	virtual void ConditionalBuildKeyMappings_Internal() const override;
