EnhancedInput: Various fixes.
Allow blueprint modification of Input Action Values, enabling implementation of custom Input Modifiers within blueprints. #jira  UE-106321
Analog actions updated every tick if non-zero, even if unmodified.  #jira UE-102962
Improved chording within a single mapping context. #jira UE-107849
Fixed Query Map Key node crash. #jira UE-101671

diff --git a/Engine/Plugins/Experimental/EnhancedInput/Source/EnhancedInput/Private/EnhancedInputLibrary.cpp b/Engine/Plugins/Experimental/EnhancedInput/Source/EnhancedInput/Private/EnhancedInputLibrary.cpp
--- a/Engine/Plugins/Experimental/EnhancedInput/Source/EnhancedInput/Private/EnhancedInputLibrary.cpp
+++ b/Engine/Plugins/Experimental/EnhancedInput/Source/EnhancedInput/Private/EnhancedInputLibrary.cpp
@@ -44,6 +44,19 @@ FInputActionValue UEnhancedInputLibrary::GetBoundActionValue(AActor* Actor, cons
 }
 
 
+void UEnhancedInputLibrary::BreakInputActionValue(FInputActionValue InActionValue, float& X, float& Y, float& Z)
+{
+	FVector AsAxis3D = InActionValue.Get<FInputActionValue::Axis3D>();
+	X = AsAxis3D.X;
+	Y = AsAxis3D.Y;
+	Z = AsAxis3D.Z;
+}
+
+FInputActionValue UEnhancedInputLibrary::MakeInputActionValue(float X, float Y, float Z, const FInputActionValue& MatchValueType)
+{
+	return FInputActionValue(MatchValueType.GetValueType(), FVector(X, Y, Z));
+}
+
 // FInputActionValue type conversions
 
 bool UEnhancedInputLibrary::Conv_InputActionValueToBool(FInputActionValue InValue)
@@ -64,4 +77,4 @@ FVector2D UEnhancedInputLibrary::Conv_InputActionValueToAxis2D(FInputActionValue
 FVector UEnhancedInputLibrary::Conv_InputActionValueToAxis3D(FInputActionValue InValue)
 {
 	return InValue.Get<FInputActionValue::Axis3D>();
-}
\ No newline at end of file
+}
diff --git a/Engine/Plugins/Experimental/EnhancedInput/Source/EnhancedInput/Private/EnhancedInputSubsystemInterface.cpp b/Engine/Plugins/Experimental/EnhancedInput/Source/EnhancedInput/Private/EnhancedInputSubsystemInterface.cpp
--- a/Engine/Plugins/Experimental/EnhancedInput/Source/EnhancedInput/Private/EnhancedInputSubsystemInterface.cpp
+++ b/Engine/Plugins/Experimental/EnhancedInput/Source/EnhancedInput/Private/EnhancedInputSubsystemInterface.cpp
@@ -95,6 +95,11 @@ EMappingQueryResult IEnhancedInputSubsystemInterface::QueryMapKeyInActiveContext
 
 EMappingQueryResult IEnhancedInputSubsystemInterface::QueryMapKeyInContextSet(const TArray<UInputMappingContext*>& PrioritizedActiveContexts, const UInputMappingContext* InputContext, const UInputAction* Action, FKey Key, TArray<FMappingQueryIssue>& OutIssues, EMappingQueryIssue BlockingIssues/* = DefaultMappingIssues::StandardFatal*/)
 {
+	if (!Action)
+	{
+		return EMappingQueryResult::Error_InvalidAction;
+	}
+
 	OutIssues.Reset();
 
 	// Report on keys being bound that don't support the action's value type.
@@ -187,7 +192,6 @@ EMappingQueryResult IEnhancedInputSubsystemInterface::QueryMapKeyInContextSet(co
 
 }
 
-
 bool IEnhancedInputSubsystemInterface::HasTriggerWith(TFunctionRef<bool(const UInputTrigger*)> TestFn, const TArray<UInputTrigger*>& Triggers)
 {
 	for (const UInputTrigger* Trigger : Triggers)
@@ -350,6 +354,70 @@ void DeepCopyPtrArray(const TArray<T*>& From, TArray<T*>& To)
 	}
 }
 
+// TODO: This should be a delegate (along with InjectChordBlockers), moving chording out of the underlying subsystem and enabling implementation of custom mapping handlers.
+TArray<FEnhancedActionKeyMapping> ReorderMappings(const TArray<FEnhancedActionKeyMapping>& UnorderedMappings)
+{
+	// Reorder mappings such that chording mappings > chorded mappings > everything else. This is used to ensure mappings within a single context are evaluated in the correct order to support chording.
+
+	TSet<const UInputAction*> ChordingActions;
+
+	// Gather all chording actions within a mapping's triggers.
+	auto GatherChordingActions = [&ChordingActions](const FEnhancedActionKeyMapping& Mapping) {
+		bool bFoundChordTrigger = false;
+		auto EvaluateTriggers = [&Mapping, &ChordingActions, &bFoundChordTrigger](const TArray<UInputTrigger*>& Triggers) {
+			for (const UInputTrigger* Trigger : Triggers)
+			{
+				if (const UInputTriggerChordAction* ChordTrigger = Cast<const UInputTriggerChordAction>(Trigger))
+				{
+					ChordingActions.Add(ChordTrigger->ChordAction);
+					bFoundChordTrigger = true;
+				}
+			}
+		};
+		EvaluateTriggers(Mapping.Triggers);
+		EvaluateTriggers(Mapping.Action->Triggers);
+		return bFoundChordTrigger;
+	};
+
+	// Split chorded mappings (second priority) from all others whilst building a list of chording actions to use for further prioritization.
+	TArray<FEnhancedActionKeyMapping> ChordedMappings;
+	TArray<FEnhancedActionKeyMapping> OtherMappings;
+	OtherMappings.Reserve(UnorderedMappings.Num());		// Mappings will most likely be Other
+	for (const FEnhancedActionKeyMapping& Mapping : UnorderedMappings)
+	{
+		TArray<FEnhancedActionKeyMapping>& MappingArray = GatherChordingActions(Mapping) ? ChordedMappings : OtherMappings;
+		MappingArray.Add(Mapping);
+	}
+
+	TArray<FEnhancedActionKeyMapping> OrderedMappings;
+	OrderedMappings.Reserve(UnorderedMappings.Num());
+
+	// Move chording mappings to the front as they need to be evaluated before chord and blocker triggers
+	// TODO: Further ordering of chording mappings may be required should one of them be chorded against another
+	auto ExtractChords = [&OrderedMappings, &ChordingActions](TArray<FEnhancedActionKeyMapping>& Mappings) {
+		for (int32 i = 0; i < Mappings.Num();)
+		{
+			if (ChordingActions.Contains(Mappings[i].Action))
+			{
+				OrderedMappings.Add(Mappings[i]);
+				Mappings.RemoveAtSwap(i);	// TODO: Do we care about reordering underlying mappings?
+			}
+			else
+			{
+				++i;
+			}
+		}
+	};
+	ExtractChords(ChordedMappings);
+	ExtractChords(OtherMappings);
+
+	OrderedMappings.Append(MoveTemp(ChordedMappings));
+	OrderedMappings.Append(MoveTemp(OtherMappings));
+	checkf(OrderedMappings.Num() == UnorderedMappings.Num(), TEXT("Number of mappings changed during reorder."));
+
+	return OrderedMappings;
+}
+
 void IEnhancedInputSubsystemInterface::RebuildControlMappings()
 {
 	if(!bMappingRebuildPending)
@@ -382,15 +450,18 @@ void IEnhancedInputSubsystemInterface::RebuildControlMappings()
 		TArray<FKey> ContextAppliedKeys;
 
 		const UInputMappingContext* MappingContext = ContextPair.Key;
-		for (const FEnhancedActionKeyMapping& Mapping : MappingContext->GetMappings())
+		TArray<FEnhancedActionKeyMapping>  OrderedMappings = ReorderMappings(MappingContext->GetMappings());
+
+		for (const FEnhancedActionKeyMapping& Mapping : OrderedMappings)
 		{
 			if (Mapping.Action && !AppliedKeys.Contains(Mapping.Key))
 			{
-				auto AnyChords = [](const UInputTrigger* Trigger) { return Cast<const UInputTriggerChordAction>(Trigger) != nullptr; };
-				bool bHasActionChords = HasTriggerWith(AnyChords, Mapping.Action->Triggers);
-				bool bHasChords = HasTriggerWith(AnyChords, Mapping.Triggers) || bHasActionChords;
+				// TODO: Wasteful query as we've already established chord state within ReorderMappings. Store TOptional bConsumeInput per mapping, allowing override? Query override via delegate?
+				auto IsChord = [](const UInputTrigger* Trigger) { return Cast<const UInputTriggerChordAction>(Trigger) != nullptr; };
+				bool bHasActionChords = HasTriggerWith(IsChord, Mapping.Action->Triggers);
+				bool bHasChords = HasTriggerWith(IsChord, Mapping.Triggers) || bHasActionChords;
 
-				// Chorded actions can't consume input or they could hide the action they are chording.
+				// Chorded actions can't consume input or they would hide the action they are chording.
 				if (!bHasChords && Mapping.Action->bConsumeInput)
 				{
 					ContextAppliedKeys.Add(Mapping.Key);
@@ -422,7 +493,7 @@ void IEnhancedInputSubsystemInterface::RebuildControlMappings()
 					{
 						for (const UInputTrigger* Trigger : Mapping.Action->Triggers)
 						{
-							if (AnyChords(Trigger))
+							if (IsChord(Trigger))
 							{
 								NewMapping.Triggers.Add(DuplicateObject(Trigger, nullptr));
 							}
diff --git a/Engine/Plugins/Experimental/EnhancedInput/Source/EnhancedInput/Private/EnhancedInputSubsystemInterfaceDebug.cpp b/Engine/Plugins/Experimental/EnhancedInput/Source/EnhancedInput/Private/EnhancedInputSubsystemInterfaceDebug.cpp
--- a/Engine/Plugins/Experimental/EnhancedInput/Source/EnhancedInput/Private/EnhancedInputSubsystemInterfaceDebug.cpp
+++ b/Engine/Plugins/Experimental/EnhancedInput/Source/EnhancedInput/Private/EnhancedInputSubsystemInterfaceDebug.cpp
@@ -157,6 +157,22 @@ void IEnhancedInputSubsystemInterface::ShowDebugInfo(UCanvas* Canvas)
 
 				ETriggerState ActionTriggerState = ETriggerState::None;
 
+				auto GetTriggerState = [](const TArray<UInputTrigger*> Triggers) {
+					FString TriggerOutput;
+					for (UInputTrigger* Trigger : Triggers)
+					{
+						if (Trigger)
+						{
+							FString State = Trigger->GetDebugState();
+							if (State.Len())
+							{
+								TriggerOutput += (TriggerOutput.IsEmpty() ? FString("") : FString(", ")) + State;
+							}
+						}
+					}
+					return !TriggerOutput.IsEmpty() ? "  Triggers: " + TriggerOutput : FString();
+				};
+
 				// TODO: Order these by key display name?
 				for (FEnhancedActionKeyMapping& Mapping : ActionMappings[Action])
 				{
@@ -206,22 +222,9 @@ void IEnhancedInputSubsystemInterface::ShowDebugInfo(UCanvas* Canvas)
 					}
 					else
 					{
-						const FEnhancedActionKeyMapping& InstancedMapping = PlayerInput->EnhancedActionMappings[Idx];
-
 						// TODO: Show state breakdown for all (Active?) triggers?
-
-						FString TriggerOutput;
-						for (UInputTrigger* Trigger : InstancedMapping.Triggers)
-						{
-							if (Trigger)
-							{
-								TriggerOutput += (TriggerOutput.IsEmpty() ? FString("") : FString(", ")) + Trigger->GetDebugState();
-							}
-						}
-						if (!TriggerOutput.IsEmpty())
-						{
-							Output += "  ( " + TriggerOutput + " )";
-						}
+						const FEnhancedActionKeyMapping& InstancedMapping = PlayerInput->EnhancedActionMappings[Idx];
+						Output += GetTriggerState(InstancedMapping.Triggers);
 					}
 
 					MappingDisplayLine.Add(FDisplayLine(DrawColor, Output));
@@ -236,6 +239,7 @@ void IEnhancedInputSubsystemInterface::ShowDebugInfo(UCanvas* Canvas)
 				{
 					ActionStr += FString::Printf(TEXT(" - %s - %.3fs (%s)"), *UEnum::GetValueAsString(TEXT("EnhancedInput.ETriggerEvent"), ActionData->GetTriggerEvent()), ActionData->GetTriggerEvent() == ETriggerEvent::Triggered ? ActionData->GetTriggeredTime() : ActionData->GetElapsedTime(), *ActionData->GetValue().ToString());
 				}
+				ActionStr += GetTriggerState(ActionData->GetTriggers());
 				DisplayDebugManager.DrawString(ActionStr);
 
 				ShowDebugActionModifiers(Canvas, Action);
diff --git a/Engine/Plugins/Experimental/EnhancedInput/Source/EnhancedInput/Private/EnhancedPlayerInput.cpp b/Engine/Plugins/Experimental/EnhancedInput/Source/EnhancedInput/Private/EnhancedPlayerInput.cpp
--- a/Engine/Plugins/Experimental/EnhancedInput/Source/EnhancedInput/Private/EnhancedPlayerInput.cpp
+++ b/Engine/Plugins/Experimental/EnhancedInput/Source/EnhancedInput/Private/EnhancedPlayerInput.cpp
@@ -114,7 +114,7 @@ void UEnhancedPlayerInput::ProcessActionMappingEvent(const UInputAction* Action,
 	{
 		if (bResetActionData)
 		{
-			ActionsWithEventsThisTick.Add(Action);			
+			ActionsWithEventsThisTick.Add(Action);
 			ActionData.Value.Reset();	// TODO: what if default value isn't 0 (e.g. bool value with negate modifier). Move reset out to a pre-pass? This may be confusing as triggering requires key interaction for value processing for performance reasons.
 		}
 
@@ -128,7 +128,7 @@ void UEnhancedPlayerInput::ProcessActionMappingEvent(const UInputAction* Action,
 		ActionData.bMappingTriggerApplied |= Triggers.Num() > 0;
 
 		// Combine values for active events only, selecting the input with the greatest magnitude for each component in each tick.
-		if (TriggerState != ETriggerState::None)
+		if(ModifiedValue.GetMagnitudeSq())
 		{
 			const int32 NumComponents = FMath::Max(1, int32(ValueType));
 			FVector Modified = ModifiedValue.Get<FVector>();
@@ -170,6 +170,7 @@ void UEnhancedPlayerInput::ProcessInputStack(const TArray<UInputComponent*>& Inp
 			const FKeyState& KeyState = KeyPair.Value;
 			// TODO: Can't just use bDownPrevious as paired axis event edges may not fire due to axial deadzoning/missing axis properties. Need to change how this is detected in PlayerInput.cpp.
 			bool bWasDown = KeyState.bDownPrevious || KeyState.EventCounts[IE_Pressed].Num() || KeyState.EventCounts[IE_Repeat].Num();
+			bWasDown |= KeyPair.Key.IsAnalog() && KeyState.RawValue.SizeSquared() != 0;	// Analog inputs should pulse every (non-zero) tick to retain compatibility with UE4.
 			KeyDownPrevious.Emplace(KeyPair.Key, bWasDown);
 		}
 	}
@@ -199,6 +200,9 @@ void UEnhancedPlayerInput::ProcessInputStack(const TArray<UInputComponent*>& Inp
 		bool bDownLastTick = KeyDownPrevious.FindRef(Mapping.Key);
 		// TODO: Can't just use bDown as paired axis event edges may not fire due to axial deadzoning/missing axis properties. Need to change how this is detected in PlayerInput.cpp.
 		bool bKeyIsDown = KeyState && (KeyState->bDown || KeyState->EventCounts[IE_Pressed].Num() || KeyState->EventCounts[IE_Repeat].Num());
+		// Analog inputs should pulse every (non-zero) tick to retain compatibility with UE4. TODO: This would be better handled at the device level.
+		bKeyIsDown |= Mapping.Key.IsAnalog() && RawKeyValue.SizeSquared() > 0;
+
 		bool bKeyIsReleased = !bKeyIsDown && bDownLastTick;
 		bool bKeyIsHeld = bKeyIsDown && bDownLastTick;
 
diff --git a/Engine/Plugins/Experimental/EnhancedInput/Source/EnhancedInput/Public/EnhancedInputLibrary.h b/Engine/Plugins/Experimental/EnhancedInput/Source/EnhancedInput/Public/EnhancedInputLibrary.h
--- a/Engine/Plugins/Experimental/EnhancedInput/Source/EnhancedInput/Public/EnhancedInputLibrary.h
+++ b/Engine/Plugins/Experimental/EnhancedInput/Source/EnhancedInput/Public/EnhancedInputLibrary.h
@@ -36,6 +36,18 @@ public:
 	static void RequestRebuildControlMappingsUsingContext(const UInputMappingContext* Context, bool bForceImmediately = false);
 
 
+	/** Breaks an ActionValue into X, Y, Z. Axes not supported by value type will be 0. */
+	UFUNCTION(BlueprintPure, Category = "Input", meta = (NativeBreakFunc))
+	static void BreakInputActionValue(FInputActionValue InActionValue, float& X, float& Y, float& Z);
+
+	/**
+	 * Builds an ActionValue from X, Y, Z. Inherits type from an existing ActionValue. Ignores axis values unused by the provided value type.
+	 * @note Intended for use in Input Modifier Modify Raw overloads to modify an existing Input Action Value.
+	 */
+	UFUNCTION(BlueprintPure, Category = "Input", meta = (Keywords = "construct build", NativeMakeFunc))
+	static FInputActionValue MakeInputActionValue(float X, float Y, float Z, const FInputActionValue& MatchValueType);
+
+
 	// Internal helper functionality
 
 	// GetInputActionvalue internal accessor function for actions that have been bound to from a UEnhancedInputComponent
@@ -58,6 +70,5 @@ public:
 
 	/** Interpret an InputActionValue as a 3D axis (Vector) input */
 	UFUNCTION(BlueprintPure, BlueprintInternalUseOnly, meta = (BlueprintAutocast))
-	static FVector Conv_InputActionValueToAxis3D(FInputActionValue InValue);
-
+	static FVector Conv_InputActionValueToAxis3D(FInputActionValue ActionValue);
 };
diff --git a/Engine/Plugins/Experimental/EnhancedInput/Source/EnhancedInput/Public/InputActionValue.h b/Engine/Plugins/Experimental/EnhancedInput/Source/EnhancedInput/Public/InputActionValue.h
--- a/Engine/Plugins/Experimental/EnhancedInput/Source/EnhancedInput/Public/InputActionValue.h
+++ b/Engine/Plugins/Experimental/EnhancedInput/Source/EnhancedInput/Public/InputActionValue.h
@@ -20,7 +20,7 @@ enum class EInputActionValueType : uint8
 };
 
 
-USTRUCT(BlueprintType)
+USTRUCT(BlueprintType, meta = (HasNativeMake = "EnhancedInput.EnhancedInputLibrary.MakeInputActionValue", HasNativeBreak = "EnhancedInput.EnhancedInputLibrary.BreakInputActionValue"))
 struct FInputActionValue
 {
 	GENERATED_BODY()
diff --git a/Engine/Plugins/Experimental/EnhancedInput/Source/EnhancedInput/Public/InputMappingQuery.h b/Engine/Plugins/Experimental/EnhancedInput/Source/EnhancedInput/Public/InputMappingQuery.h
--- a/Engine/Plugins/Experimental/EnhancedInput/Source/EnhancedInput/Public/InputMappingQuery.h
+++ b/Engine/Plugins/Experimental/EnhancedInput/Source/EnhancedInput/Public/InputMappingQuery.h
@@ -16,6 +16,9 @@ enum class EMappingQueryResult : uint8
 	// Query failed because the input context being queried against is not part of the active context list.
 	Error_InputContextNotInActiveContexts,
 
+	// Query failed because the action being queried against is None/null
+	Error_InvalidAction,
+
 	// Mapping cannot be applied due to blocking issues. Check OutIssues for details.
 	NotMappable,
 
diff --git a/Engine/Plugins/Experimental/EnhancedInput/Source/EnhancedInput/Public/InputTriggers.h b/Engine/Plugins/Experimental/EnhancedInput/Source/EnhancedInput/Public/InputTriggers.h
--- a/Engine/Plugins/Experimental/EnhancedInput/Source/EnhancedInput/Public/InputTriggers.h
+++ b/Engine/Plugins/Experimental/EnhancedInput/Source/EnhancedInput/Public/InputTriggers.h
@@ -180,6 +180,7 @@ class UInputTriggerPressed final : public UInputTrigger
 protected:
 
 	virtual ETriggerState UpdateState_Implementation(const UEnhancedPlayerInput* PlayerInput, FInputActionValue ModifiedValue, float DeltaTime) override;
+	virtual FString GetDebugState() const { return IsActuated(LastValue) ? FString(TEXT("Pressed:Held")) : FString(); }
 };
 
 
@@ -195,6 +196,7 @@ class UInputTriggerReleased final : public UInputTrigger
 protected:
 
 	virtual ETriggerState UpdateState_Implementation(const UEnhancedPlayerInput* PlayerInput, FInputActionValue ModifiedValue, float DeltaTime) override;
+	virtual FString GetDebugState() const { return IsActuated(LastValue) ? FString(TEXT("Released:Held")) : FString(); }
 };
 
 
@@ -297,7 +299,6 @@ class UInputTriggerChordBlocker final : public UInputTriggerChordAction
 
 protected:
 	virtual ETriggerType GetTriggerType_Implementation() const override { return ETriggerType::Blocker; }
-
 };
 
 
diff --git a/Engine/Plugins/Experimental/EnhancedInput/Source/InputEditor/Private/Tests/InputTriggerTest.cpp b/Engine/Plugins/Experimental/EnhancedInput/Source/InputEditor/Private/Tests/InputTriggerTest.cpp
--- a/Engine/Plugins/Experimental/EnhancedInput/Source/InputEditor/Private/Tests/InputTriggerTest.cpp
+++ b/Engine/Plugins/Experimental/EnhancedInput/Source/InputEditor/Private/Tests/InputTriggerTest.cpp
@@ -357,15 +357,14 @@ bool FInputTriggerTapTest::RunTest(const FString& Parameters)
 }
 
 
-IMPLEMENT_SIMPLE_AUTOMATION_TEST(FInputTriggerChordedActionsTest, "Input.Triggers.ChordedActions", BasicTriggerTestFlags)
+IMPLEMENT_SIMPLE_AUTOMATION_TEST(FInputTriggerChordsMultiContextTest, "Input.Triggers.Chords.MultiContext", BasicTriggerTestFlags)
 
-bool FInputTriggerChordedActionsTest::RunTest(const FString& Parameters)
+bool FInputTriggerChordsMultiContextTest::RunTest(const FString& Parameters)
 {
-	// Test chording changing the triggered action
-
+	// Test chords work when the chord action is in a higher priority context
 	FKey ChordKey = TestKey2;
-	FName BaseAction = TEXT("BaseAction");				// Base action e.g. Jump
-	FName ChordedAction = TEXT("ChordedAction");		// Special case action  e.g. Back flip
+	FName BaseAction = TEXT("BaseAction");				// Base action
+	FName ChordedAction = TEXT("ChordedAction");			// Chord triggered action
 	FName ChordingAction = TEXT("ChordingAction");		// Chording action driving special case e.g. ShiftDown/AcrobaticModifier
 
 	UWorld* World =
@@ -379,94 +378,165 @@ bool FInputTriggerChordedActionsTest::RunTest(const FString& Parameters)
 	AND(AnInputContextIsAppliedToAPlayer(Data, BaseContext, 1));
 	AND(AnInputContextIsAppliedToAPlayer(Data, ChordContext, 100));
 
-	// Set up actions
-	AND(AnInputAction(Data, BaseAction, EInputActionValueType::Boolean));
-	AND(AnInputAction(Data, ChordedAction, EInputActionValueType::Boolean));
+	// Set up action
+	AND(AnInputAction(Data, BaseAction, EInputActionValueType::Axis1D));
+	AND(AnInputAction(Data, ChordedAction, EInputActionValueType::Axis1D));
 
 	// Set up the chording action (modifier key action)
 	AND(UInputAction * ChordingActionPtr = AnInputAction(Data, ChordingAction, EInputActionValueType::Boolean));
 
-	// Apply a chord action trigger to the chorded action
+	// Bind the chording modifier in the high priority context
+	AND(AnActionIsMappedToAKey(Data, ChordContext, ChordingAction, ChordKey));
+
+	// Bind the action to the same key in both contexts
+	AND(AnActionIsMappedToAKey(Data, BaseContext, BaseAction, TestAxis));
+	AND(AnActionIsMappedToAKey(Data, ChordContext, ChordedAction, TestAxis));
+
+	// But the chorded version inverts the result
+	AND(AModifierIsAppliedToAnActionMapping(Data, NewObject<UInputModifierNegate>(), ChordContext, ChordedAction, TestAxis));
+
+	// Apply a chord action trigger to the chorded mapping
 	UInputTriggerChordAction* Trigger = NewObject<UInputTriggerChordAction>();
 	Trigger->ChordAction = ChordingActionPtr;
-	AND(ATriggerIsAppliedToAnAction(Data, Trigger, ChordedAction));
+	AND(ATriggerIsAppliedToAnActionMapping(Data, Trigger, ChordContext, ChordedAction, TestAxis));
 
-	// Bind the chording modifier
-	AND(AnActionIsMappedToAKey(Data, ChordContext, ChordingAction, ChordKey));
 
-	// Bind both actions to the same key
-	AND(AnActionIsMappedToAKey(Data, BaseContext, BaseAction, TestKey));
-	AND(AnActionIsMappedToAKey(Data, ChordContext, ChordedAction, TestKey));
+	TRIGGER_SUBTEST("With chord key pressed neither main action triggers, but chording action does")
+	{
+		WHEN(AKeyIsActuated(Data, ChordKey));
+		AND(InputIsTicked(Data));
+		THEN(PressingKeyDoesNotTrigger(Data, BaseAction));
+		ANDALSO(PressingKeyDoesNotTrigger(Data, ChordedAction));
+		ANDALSO(PressingKeyTriggersAction(Data, ChordingAction));
+	}
 
-	// With chord key pressed no main actions trigger, but chording action does
-	WHEN(AKeyIsActuated(Data, ChordKey));
-	AND(InputIsTicked(Data));
-	THEN(PressingKeyDoesNotTrigger(Data, BaseAction));
-	ANDALSO(PressingKeyTriggersAction(Data, ChordingAction));
-	ANDALSO(PressingKeyDoesNotTrigger(Data, ChordedAction));
-
-	// Switching to test key the base action only triggers
-	WHEN(AKeyIsReleased(Data, ChordKey));
-	AND(AKeyIsActuated(Data, TestKey));
-	AND(InputIsTicked(Data));
-	THEN(PressingKeyTriggersAction(Data, BaseAction));
-	ANDALSO(ReleasingKeyTriggersCompleted(Data, ChordingAction));
-	ANDALSO(ReleasingKeyDoesNotTrigger(Data, ChordedAction));
-
-	// Depressing chord key triggers chorded action, and ends base action
-	WHEN(AKeyIsActuated(Data, ChordKey));
-	AND(InputIsTicked(Data));
-	THEN(PressingKeyTriggersCompleted(Data, BaseAction));
-	ANDALSO(PressingKeyTriggersAction(Data, ChordingAction));
-	ANDALSO(PressingKeyTriggersAction(Data, ChordedAction));
-
-	// Releasing chord key returns to to base only
-	WHEN(AKeyIsReleased(Data, ChordKey));
-	AND(InputIsTicked(Data));
-	THEN(ReleasingKeyTriggersAction(Data, BaseAction));
-	ANDALSO(ReleasingKeyTriggersCompleted(Data, ChordingAction));
-	ANDALSO(ReleasingKeyTriggersCompleted(Data, ChordedAction));
-
-	WHEN(AKeyIsReleased(Data, TestKey));
-	AND(InputIsTicked(Data));
-	THEN(ReleasingKeyTriggersCompleted(Data, BaseAction));
-	ANDALSO(ReleasingKeyDoesNotTrigger(Data, ChordingAction));
-	ANDALSO(ReleasingKeyDoesNotTrigger(Data, ChordedAction));
+	const float AxisValue = 0.5f;
 
+	TRIGGER_SUBTEST("Switching to test key the action supplies the unmodified value")
+	{
+		WHEN(AKeyIsReleased(Data, ChordKey));
+		AND(AKeyIsActuated(Data, TestAxis, AxisValue));
+		AND(InputIsTicked(Data));
+		THEN(PressingKeyTriggersAction(Data, BaseAction));
+		ANDALSO(ReleasingKeyTriggersCompleted(Data, ChordingAction));
+		AND(TestEqual(TEXT("BaseAction"), FInputTestHelper::GetTriggered<float>(Data, BaseAction), AxisValue));
+	}
+
+	TRIGGER_SUBTEST("Depressing chord key triggers chorded action modified value")
+	{
+		WHEN(AKeyIsActuated(Data, ChordKey));
+		AND(InputIsTicked(Data));
+		THEN(PressingKeyTriggersAction(Data, ChordedAction));
+		ANDALSO(ReleasingKeyTriggersCompleted(Data, BaseAction));
+		AND(TestEqual(TEXT("BaseAction"), FInputTestHelper::GetTriggered<float>(Data, ChordedAction), -AxisValue));
+	}
 
 	return true;
 }
 
-IMPLEMENT_SIMPLE_AUTOMATION_TEST(FInputTriggerChordedModifiersTest, "Input.Triggers.ChordedModifiers", BasicTriggerTestFlags)
+IMPLEMENT_SIMPLE_AUTOMATION_TEST(FInputTriggerChordsSingleContextTest, "Input.Triggers.Chords.SingleContext", BasicTriggerTestFlags)
 
-bool FInputTriggerChordedModifiersTest::RunTest(const FString& Parameters)
+bool FInputTriggerChordsSingleContextTest::RunTest(const FString& Parameters)
 {
-	// Test applying a different set of modifiers to an action based on chords
-	// Unchorded mapping with no modifier
-	// Chorded mapping with negate modifiers
+	// Test chords work when all base, chorded, and chording actions are in the same context (they are processed in the correct order)
+	FKey ChordKey = TestKey2;
+	FName BaseAction = TEXT("BaseAction");				// Base action
+	FName ChordedAction = TEXT("ChordedAction");			// Chord triggered action
+	FName ChordingAction = TEXT("ChordingAction");		// Chording action driving special case e.g. ShiftDown/AcrobaticModifier
+
+	UWorld* World =
+		GIVEN(AnEmptyWorld());
+
+	// Initialise
+	UControllablePlayer& Data =
+		AND(AControllablePlayer(World));
 
-	auto GetActionValue = [](UControllablePlayer & Data, FName ActionName)
+	FName SingleContext = TEXT("Context");
+	AND(AnInputContextIsAppliedToAPlayer(Data, SingleContext, 1));
+
+	// Set up action
+	AND(AnInputAction(Data, BaseAction, EInputActionValueType::Axis1D));
+	AND(AnInputAction(Data, ChordedAction, EInputActionValueType::Axis1D));
+
+	// Set up the chording action (modifier key action)
+	AND(UInputAction * ChordingActionPtr = AnInputAction(Data, ChordingAction, EInputActionValueType::Boolean));
+
+	// Bind the chording modifier
+	AND(AnActionIsMappedToAKey(Data, SingleContext, ChordingAction, ChordKey));
+
+	// Bind the actions to the same key
+	AND(AnActionIsMappedToAKey(Data, SingleContext, BaseAction, TestAxis));
+	AND(AnActionIsMappedToAKey(Data, SingleContext, ChordedAction, TestAxis));
+
+	// But the chorded version inverts the result
+	AND(AModifierIsAppliedToAnActionMapping(Data, NewObject<UInputModifierNegate>(), SingleContext, ChordedAction, TestAxis));
+
+	// Apply a chord action trigger to the chorded mapping
+	UInputTriggerChordAction* Trigger = NewObject<UInputTriggerChordAction>();
+	Trigger->ChordAction = ChordingActionPtr;
+	AND(ATriggerIsAppliedToAnActionMapping(Data, Trigger, SingleContext, ChordedAction, TestAxis));
+
+
+	TRIGGER_SUBTEST("With chord key pressed neither main action triggers, but chording action does")
 	{
-		return FInputTestHelper::GetActionData(Data, ActionName).GetValue().Get<float>();
-	};
+		WHEN(AKeyIsActuated(Data, ChordKey));
+		AND(InputIsTicked(Data));
+		THEN(PressingKeyDoesNotTrigger(Data, BaseAction));
+		ANDALSO(PressingKeyDoesNotTrigger(Data, ChordedAction));
+		ANDALSO(PressingKeyTriggersAction(Data, ChordingAction));
+	}
+
+	const float AxisValue = 0.5f;
+
+	TRIGGER_SUBTEST("Switching to test key the action supplies the unmodified value")
+	{
+		WHEN(AKeyIsReleased(Data, ChordKey));
+		AND(AKeyIsActuated(Data, TestAxis, AxisValue));
+		AND(InputIsTicked(Data));
+		THEN(PressingKeyTriggersAction(Data, BaseAction));
+		ANDALSO(ReleasingKeyTriggersCompleted(Data, ChordingAction));
+		AND(TestEqual(TEXT("BaseAction"), FInputTestHelper::GetTriggered<float>(Data, BaseAction), AxisValue));
+	}
 
-	FKey ChordKey = TestKey;
+	TRIGGER_SUBTEST("Depressing chord key triggers chorded action modified value")
+	{
+		WHEN(AKeyIsActuated(Data, ChordKey));
+		AND(InputIsTicked(Data));
+		THEN(PressingKeyTriggersAction(Data, ChordedAction));
+		ANDALSO(ReleasingKeyTriggersCompleted(Data, BaseAction));
+		AND(TestEqual(TEXT("BaseAction"), FInputTestHelper::GetTriggered<float>(Data, ChordedAction), -AxisValue));
+	}
+
+	return true;
+}
+
+IMPLEMENT_SIMPLE_AUTOMATION_TEST(FInputTriggerChordsPressedTriggersTest, "Input.Triggers.Chords.WithPressedTriggers", BasicTriggerTestFlags)
+
+bool FInputTriggerChordsPressedTriggersTest::RunTest(const FString& Parameters)
+{
+	// Test chord behavior with pressed triggers
+	// Expected: Main key trigger state should be retained by both base and chorded action, across any chord key state transitions.
+	// Pressing or releasing the chord key shouldn't cause any action to trigger by itself (Note: triggering would continue to occur for a down trigger).
+
+	FKey ChordKey = TestKey2;
 	FName BaseAction = TEXT("BaseAction");				// Base action
+	FName ChordedAction = TEXT("ChordedAction");			// Chord triggered action
 	FName ChordingAction = TEXT("ChordingAction");		// Chording action driving special case e.g. ShiftDown/AcrobaticModifier
 
 	UWorld* World =
-	GIVEN(AnEmptyWorld());
+		GIVEN(AnEmptyWorld());
 
 	// Initialise
 	UControllablePlayer& Data =
-	AND(AControllablePlayer(World));
+		AND(AControllablePlayer(World));
 
 	FName BaseContext = TEXT("BaseContext"), ChordContext = TEXT("ChordContext");
 	AND(AnInputContextIsAppliedToAPlayer(Data, BaseContext, 1));
 	AND(AnInputContextIsAppliedToAPlayer(Data, ChordContext, 100));
 
 	// Set up action
-	AND(AnInputAction(Data, BaseAction, EInputActionValueType::Axis1D));
+	AND(AnInputAction(Data, BaseAction, EInputActionValueType::Boolean));
+	AND(AnInputAction(Data, ChordedAction, EInputActionValueType::Boolean));
 
 	// Set up the chording action (modifier key action)
 	AND(UInputAction * ChordingActionPtr = AnInputAction(Data, ChordingAction, EInputActionValueType::Boolean));
@@ -474,48 +544,59 @@ bool FInputTriggerChordedModifiersTest::RunTest(const FString& Parameters)
 	// Bind the chording modifier in the high priority context
 	AND(AnActionIsMappedToAKey(Data, ChordContext, ChordingAction, ChordKey));
 
-	// Bind the action to the same key in both contexts
-	AND(AnActionIsMappedToAKey(Data, BaseContext, BaseAction, TestAxis));
-	AND(AnActionIsMappedToAKey(Data, ChordContext, BaseAction, TestAxis));
+	// Bind the actions to the same key in both contexts
+	AND(AnActionIsMappedToAKey(Data, BaseContext, BaseAction, TestKey));
+	AND(AnActionIsMappedToAKey(Data, ChordContext, ChordedAction, TestKey));
 
-	// But the chorded version inverts the result
-	AND(AModifierIsAppliedToAnActionMapping(Data, NewObject<UInputModifierNegate>(), ChordContext, BaseAction, TestAxis));
+	// Apply pressed triggers to both actions
+	AND(ATriggerIsAppliedToAnAction(Data, NewObject<UInputTriggerPressed>(), BaseAction));
+	AND(ATriggerIsAppliedToAnAction(Data, NewObject<UInputTriggerPressed>(), ChordedAction));
 
 	// Apply a chord action trigger to the chorded mapping
 	UInputTriggerChordAction* Trigger = NewObject<UInputTriggerChordAction>();
 	Trigger->ChordAction = ChordingActionPtr;
-	AND(ATriggerIsAppliedToAnActionMapping(Data, Trigger, ChordContext, BaseAction, TestAxis));
+	AND(ATriggerIsAppliedToAnAction(Data, Trigger, ChordedAction));
 
+	TRIGGER_SUBTEST("Pressing key triggers base action")
+	{
+		WHEN(AKeyIsActuated(Data, TestKey));
+		AND(InputIsTicked(Data));
+		THEN(PressingKeyTriggersAction(Data, BaseAction));
+		ANDALSO(PressingKeyDoesNotTrigger(Data, ChordedAction));
+	}
 
-	// With chord key pressed main action does not trigger, but chording action does
-	WHEN(AKeyIsActuated(Data, ChordKey));
-	AND(InputIsTicked(Data));
-	THEN(PressingKeyDoesNotTrigger(Data, BaseAction));
-	ANDALSO(PressingKeyTriggersAction(Data, ChordingAction));
+	TRIGGER_SUBTEST("Pressing chord key does not trigger chorded action, but stops base")
+	{
+		WHEN(AKeyIsActuated(Data, ChordKey));
+		AND(InputIsTicked(Data));
+		THEN(PressingKeyTriggersStarted(Data, ChordedAction));	// Begins tracking trigger...	// TODO: Started -> Permanent Ongoing. The implicit chord action is true, but explict Pressed blocks. Make chord action a 4th type? ImplicitBlocker?
+		THEN(!PressingKeyTriggersAction(Data, ChordedAction));	// but does not fire
+		ANDALSO(PressingKeyTriggersCompleted(Data, BaseAction));
+	}
 
-	const float AxisValue = 0.5f;
+	// Release main key
+	AKeyIsReleased(Data, TestKey);
+	InputIsTicked(Data);
 
-	// Switching to test key the action supplies the unmodified value
-	WHEN(AKeyIsReleased(Data, ChordKey));
-	AND(AKeyIsActuated(Data, TestAxis, AxisValue));
-	AND(InputIsTicked(Data));
-	THEN(PressingKeyTriggersAction(Data, BaseAction));
-	ANDALSO(ReleasingKeyTriggersCompleted(Data, ChordingAction));
-	AND(TestEqual(TEXT("BaseAction"), FInputTestHelper::GetTriggered<float>(Data, BaseAction), AxisValue));
-
-	// Depressing chord key triggers chorded action modified value
-	WHEN(AKeyIsActuated(Data, ChordKey));
-	AND(InputIsTicked(Data));
-	THEN(PressingKeyTriggersAction(Data, BaseAction));
-	ANDALSO(PressingKeyTriggersAction(Data, ChordingAction));
-	AND(TestEqual(TEXT("BaseAction"), FInputTestHelper::GetTriggered<float>(Data, BaseAction), -AxisValue));
+	TRIGGER_SUBTEST("Pressing key again triggers chorded action only")
+	{
+		WHEN(AKeyIsActuated(Data, TestKey));
+		AND(InputIsTicked(Data));
+		THEN(PressingKeyTriggersAction(Data, ChordedAction));
+		ANDALSO(PressingKeyDoesNotTrigger(Data, BaseAction));
+	}
+
+	TRIGGER_SUBTEST("Releasing chord key stops chorded action but does not trigger base action")
+	{
+		WHEN(AKeyIsReleased(Data, ChordKey));
+		AND(InputIsTicked(Data));
+		THEN(ReleasingKeyTriggersCompleted(Data, ChordedAction));
+		ANDALSO(ReleasingKeyDoesNotTrigger(Data, BaseAction));
+	}
 
 	return true;
 }
 
-
-
-
 // TODO: Action level triggers (simple repeat of device level tests)
 // TODO: Variable frame delta tests
 // TODO: ActionEventData tests (timing, summed values, etc)
\ No newline at end of file
