EnhancedInput: Separate delegate execution from mapping evaluation, allowing delegates to safely alter active mappings.

diff --git a/Engine/Plugins/Experimental/EnhancedInput/Source/EnhancedInput/Private/EnhancedInputComponent.cpp b/Engine/Plugins/Experimental/EnhancedInput/Source/EnhancedInput/Private/EnhancedInputComponent.cpp
--- a/Engine/Plugins/Experimental/EnhancedInput/Source/EnhancedInput/Private/EnhancedInputComponent.cpp
+++ b/Engine/Plugins/Experimental/EnhancedInput/Source/EnhancedInput/Private/EnhancedInputComponent.cpp
@@ -77,11 +77,11 @@ FInputActionValue UEnhancedInputComponent::GetBoundActionValue(const UInputActio
 	return FInputActionValue(Action->ValueType, FVector::ZeroVector);
 }
 
+// Must be in C++ to avoid duplicate statics across execution units
+static uint32 GInputBindingHandle = 1;
 
 FInputBindingHandle::FInputBindingHandle()
 {
 	// Handles are shared between all binding types
-	// Must be in C++ to avoid duplicate statics across execution units
-	static uint32 GHandle = 1;
-	Handle = GHandle++;
+	Handle = GInputBindingHandle++;
 }
\ No newline at end of file
diff --git a/Engine/Plugins/Experimental/EnhancedInput/Source/EnhancedInput/Private/EnhancedPlayerInput.cpp b/Engine/Plugins/Experimental/EnhancedInput/Source/EnhancedInput/Private/EnhancedPlayerInput.cpp
--- a/Engine/Plugins/Experimental/EnhancedInput/Source/EnhancedInput/Private/EnhancedPlayerInput.cpp
+++ b/Engine/Plugins/Experimental/EnhancedInput/Source/EnhancedInput/Private/EnhancedPlayerInput.cpp
@@ -312,6 +312,7 @@ void UEnhancedPlayerInput::ProcessInputStack(const TArray<UInputComponent*>& Inp
 		}
 
 		// Trigger bound event delegates
+		static TArray<TUniquePtr<FEnhancedInputActionEventBinding>> TriggeredDelegates;
 		for (const TUniquePtr<FEnhancedInputActionEventBinding>& Binding : IC->GetActionEventBindings())
 		{
 			// PERF: Lots of map lookups! Group EnhancedActionBindings by Action?
@@ -321,11 +322,23 @@ void UEnhancedPlayerInput::ProcessInputStack(const TArray<UInputComponent*>& Inp
 				if (ActionData->TriggerEvent == Binding->GetTriggerEvent() ||
 					(Binding->GetTriggerEvent() == ETriggerEvent::Started && ActionData->TriggerEventInternal == ETriggerEventInternal::StartedAndTriggered))	// Triggering in a single tick should also fire the started event.
 				{
-					Binding->Execute(*ActionData);
+					// Record intent to trigger
+					TriggeredDelegates.Add(Binding->Clone());
 				}
 			}
 		}
 
+		// Action all delegates that triggered this tick, in the order in which they triggered.
+		for (TUniquePtr<FEnhancedInputActionEventBinding>& Delegate : TriggeredDelegates)
+		{
+			// Search for the action instance data a second time as a previous delegate call may have deleted it.
+			if (const FInputActionInstance* ActionData = FindActionInstanceData(Delegate->GetAction()))
+			{
+				Delegate->Execute(*ActionData);
+			}
+		}
+		TriggeredDelegates.Reset();
+
 		// Update action value bindings
 		for (const FEnhancedInputActionValueBinding& Binding : IC->GetActionValueBindings())
 		{
@@ -342,6 +355,7 @@ void UEnhancedPlayerInput::ProcessInputStack(const TArray<UInputComponent*>& Inp
 		// No support for the 'Any Key' concept. Explicit key binds only.
 		// They will always fire, and cannot mask each other or action bindings (i.e. no bConsumeInput option)
 		// Chords are supported, but there is no chord masking protection. Exact chord combinations must be met. So a binding of Ctrl + A will not fire if Ctrl + Alt + A is pressed.
+		static TArray<TUniquePtr<FInputDebugKeyBinding>> TriggeredDebugDelegates;
 		for (const TUniquePtr<FInputDebugKeyBinding>& KeyBinding : IC->GetDebugKeyBindings())
 		{
 			TRACE_CPUPROFILER_EVENT_SCOPE(EnhPIS_DebugKeys);
@@ -361,11 +375,19 @@ void UEnhancedPlayerInput::ProcessInputStack(const TArray<UInputComponent*>& Inp
 					FKeyState* KeyState = KeyStateMap.Find(KeyBinding->Chord.Key);
 					if (KeyState && KeyState->EventCounts[KeyBinding->KeyEvent].Num() > 0)
 					{
-						KeyBinding->Execute();
+						// Record intent to trigger
+						TriggeredDebugDelegates.Add(KeyBinding->Clone());
 					}
 				}
 			}
 		}
+
+		// Action all debug delegates that triggered this tick, in the order in which they triggered.
+		for (TUniquePtr<FInputDebugKeyBinding>& Delegate : TriggeredDebugDelegates)
+		{
+			Delegate->Execute();
+		}
+		TriggeredDebugDelegates.Reset();
 #endif
 
 
diff --git a/Engine/Plugins/Experimental/EnhancedInput/Source/EnhancedInput/Public/EnhancedInputComponent.h b/Engine/Plugins/Experimental/EnhancedInput/Source/EnhancedInput/Public/EnhancedInputComponent.h
--- a/Engine/Plugins/Experimental/EnhancedInput/Source/EnhancedInput/Public/EnhancedInputComponent.h
+++ b/Engine/Plugins/Experimental/EnhancedInput/Source/EnhancedInput/Public/EnhancedInputComponent.h
@@ -91,6 +91,8 @@ public:
 };
 
 
+// Used to force clone constructor calls only
+enum class EInputBindingClone : uint8 { ForceClone };
 
 /** A basic binding unique identifier */
 struct FInputBindingHandle
@@ -100,6 +102,7 @@ private:
 
 protected:
 	ENHANCEDINPUT_API FInputBindingHandle();	// Generates a handle
+	FInputBindingHandle(const FInputBindingHandle& CloneFrom, EInputBindingClone) : Handle(CloneFrom.Handle) {}	// Clones a handle
 
 public:
 	virtual ~FInputBindingHandle() = default;
@@ -123,6 +126,10 @@ private:
 	/** Trigger event that raises the delegate */
 	ETriggerEvent TriggerEvent = ETriggerEvent::None;
 
+protected:
+	// Clone constructor
+	FEnhancedInputActionEventBinding(const FEnhancedInputActionEventBinding& CloneFrom, EInputBindingClone Clone) : FInputBindingHandle(CloneFrom, Clone), Action(CloneFrom.Action), TriggerEvent(CloneFrom.TriggerEvent) {}
+
 public:
 	FEnhancedInputActionEventBinding() = default;
 	FEnhancedInputActionEventBinding(const UInputAction* InAction, ETriggerEvent InTriggerEvent) : Action(InAction), TriggerEvent(InTriggerEvent) {}
@@ -131,6 +138,7 @@ public:
 	ETriggerEvent GetTriggerEvent() const { return TriggerEvent; }
 
 	virtual void Execute(const FInputActionInstance& ActionData) const = 0;
+	virtual TUniquePtr<FEnhancedInputActionEventBinding> Clone() const = 0;
 };
 
 /** Binds an action value for later reference. CurrentValue will be kept up to date with the value of the bound action */
@@ -156,6 +164,16 @@ public:
 /** Binds a delegate to an event on a key chord. */
 struct FInputDebugKeyBinding : public FInputBindingHandle
 {
+protected:
+	// Clone constructor
+	FInputDebugKeyBinding(const FInputDebugKeyBinding& CloneFrom, EInputBindingClone Clone)
+		: FInputBindingHandle(CloneFrom, Clone)
+		, KeyEvent(CloneFrom.KeyEvent)
+		, bExecuteWhenPaused(CloneFrom.bExecuteWhenPaused)
+		, Chord(CloneFrom.Chord)
+	{ }
+
+public:
 	/** Key event to bind it to (e.g. pressed, released, double click) */
 	TEnumAsByte<EInputEvent> KeyEvent = EInputEvent::IE_Pressed;
 
@@ -174,6 +192,7 @@ struct FInputDebugKeyBinding : public FInputBindingHandle
 	{ }
 
 	virtual void Execute() const = 0;
+	virtual TUniquePtr<FInputDebugKeyBinding> Clone() const = 0;
 };
 // TODO:: Add FInputDebugKeyValueBinding?
 
@@ -188,11 +207,17 @@ struct FInputDebugKeyBinding : public FInputBindingHandle
 template<typename TSignature>
 struct FEnhancedInputActionEventDelegateBinding : FEnhancedInputActionEventBinding
 {
+private:
+	FEnhancedInputActionEventDelegateBinding(const FEnhancedInputActionEventDelegateBinding<TSignature>& CloneFrom, EInputBindingClone Clone) : FEnhancedInputActionEventBinding(CloneFrom, Clone), Delegate(CloneFrom.Delegate) {}
 public:
 	FEnhancedInputActionEventDelegateBinding(const UInputAction* Action, ETriggerEvent InTriggerEvent) : FEnhancedInputActionEventBinding(Action, InTriggerEvent) {}
 
 	// Implemented below.
 	virtual void Execute(const FInputActionInstance& ActionData) const override;
+	virtual TUniquePtr<FEnhancedInputActionEventBinding> Clone() const override
+	{
+		return TUniquePtr<FEnhancedInputActionEventBinding>(new FEnhancedInputActionEventDelegateBinding<TSignature>(*this, EInputBindingClone::ForceClone));
+	}
 
 	TEnhancedInputUnifiedDelegate<TSignature> Delegate;
 };
@@ -200,12 +225,18 @@ public:
 template<typename TSignature>
 struct FInputDebugKeyDelegateBinding : FInputDebugKeyBinding
 {
+private:
+	FInputDebugKeyDelegateBinding(const FInputDebugKeyDelegateBinding<TSignature>& CloneFrom, EInputBindingClone Clone) : FInputDebugKeyBinding(CloneFrom, Clone), Delegate(CloneFrom.Delegate) {}
 public:
 	FInputDebugKeyDelegateBinding(const FInputChord Chord, const EInputEvent KeyEvent, bool bExecuteWhenPaused) : FInputDebugKeyBinding(Chord, KeyEvent, bExecuteWhenPaused) {}
 
 	virtual void Execute() const override
 	{
-		Delegate.Execute(Chord.Key);	// TODO: Remove FKey param? We don't support AnyKey, so it isn't terribly useful.
+		Delegate.Execute(Chord.Key);	// TODO: Remove FKey param? We don't support AnyKey, so it isn't terribly useful. TODO: Provide key raw value param to support axis keys.
+	}
+	virtual TUniquePtr<FInputDebugKeyBinding> Clone() const override 
+	{
+		return TUniquePtr<FInputDebugKeyBinding>(new FInputDebugKeyDelegateBinding<TSignature>(*this, EInputBindingClone::ForceClone));
 	}
 
 	TEnhancedInputUnifiedDelegate<TSignature> Delegate;
