Fixed unitialized struct var warning in EnhancedInputActionDelegateBinding

diff --git a/Engine/Plugins/Experimental/EnhancedInput/Source/EnhancedInput/Public/EnhancedInputActionDelegateBinding.h b/Engine/Plugins/Experimental/EnhancedInput/Source/EnhancedInput/Public/EnhancedInputActionDelegateBinding.h
--- a/Engine/Plugins/Experimental/EnhancedInput/Source/EnhancedInput/Public/EnhancedInputActionDelegateBinding.h
+++ b/Engine/Plugins/Experimental/EnhancedInput/Source/EnhancedInput/Public/EnhancedInputActionDelegateBinding.h
@@ -17,7 +17,7 @@ struct ENHANCEDINPUT_API FBlueprintEnhancedInputActionBinding
 	GENERATED_USTRUCT_BODY()
 
 	UPROPERTY()
-	const UInputAction* InputAction;
+	const UInputAction* InputAction = nullptr;
 
 	UPROPERTY()
 	ETriggerEvent TriggerEvent = ETriggerEvent::None;
