Fix crash accessing out of scope lambda reference when rebuilding enhanced input action lists for context menu.

diff --git a/Engine/Plugins/Experimental/EnhancedInput/Source/InputBlueprintNodes/Private/K2Node_EnhancedInputAction.cpp b/Engine/Plugins/Experimental/EnhancedInput/Source/InputBlueprintNodes/Private/K2Node_EnhancedInputAction.cpp
--- a/Engine/Plugins/Experimental/EnhancedInput/Source/InputBlueprintNodes/Private/K2Node_EnhancedInputAction.cpp
+++ b/Engine/Plugins/Experimental/EnhancedInput/Source/InputBlueprintNodes/Private/K2Node_EnhancedInputAction.cpp
@@ -287,29 +287,28 @@ void UK2Node_EnhancedInputAction::GetMenuActions(FBlueprintActionDatabaseRegistr
 	// Do a first time registration using the node's class to pull in all existing actions
 	if (ActionRegistrar.IsOpenForRegistration(GetClass()))
 	{
-		auto OnAssetsLoaded = [this, &ActionRegistrar, CustomizeInputNodeLambda]() {
-			TArray<FAssetData> ActionAssets;
-			FModuleManager::LoadModuleChecked<FAssetRegistryModule>("AssetRegistry").Get().GetAssetsByClass(UInputAction::StaticClass()->GetFName(), ActionAssets);
+		IAssetRegistry& AssetRegistry = FModuleManager::LoadModuleChecked<FAssetRegistryModule>("AssetRegistry").Get();
 
-			for (const FAssetData& ActionAsset : ActionAssets)
+		static bool bRegisterOnce = true;
+		if (bRegisterOnce)
+		{
+			bRegisterOnce = false;
+			if (AssetRegistry.IsLoadingAssets())
 			{
-				UBlueprintNodeSpawner* NodeSpawner = UBlueprintNodeSpawner::Create(GetClass());
-				check(NodeSpawner != nullptr);
-
-				const UInputAction* Action = Cast<const UInputAction>(ActionAsset.GetAsset());
-				NodeSpawner->CustomizeNodeDelegate = UBlueprintNodeSpawner::FCustomizeNodeDelegate::CreateStatic(CustomizeInputNodeLambda, TWeakObjectPtr<const UInputAction>(Action));
-				ActionRegistrar.AddBlueprintAction(Action, NodeSpawner);
+				AssetRegistry.OnFilesLoaded().AddLambda([]() { FBlueprintActionDatabase::Get().RefreshClassActions(StaticClass()); });
 			}
-		};
-
-		IAssetRegistry& AssetRegistry = FModuleManager::LoadModuleChecked<FAssetRegistryModule>("AssetRegistry").Get();
-		if (AssetRegistry.IsLoadingAssets())
-		{
-			AssetRegistry.OnFilesLoaded().AddLambda(OnAssetsLoaded);
 		}
-		else
+
+		TArray<FAssetData> ActionAssets;
+		AssetRegistry.GetAssetsByClass(UInputAction::StaticClass()->GetFName(), ActionAssets);
+		for (const FAssetData& ActionAsset : ActionAssets)
 		{
-			OnAssetsLoaded();
+			UBlueprintNodeSpawner* NodeSpawner = UBlueprintNodeSpawner::Create(GetClass());
+			check(NodeSpawner != nullptr);
+
+			const UInputAction* Action = Cast<const UInputAction>(ActionAsset.GetAsset());
+			NodeSpawner->CustomizeNodeDelegate = UBlueprintNodeSpawner::FCustomizeNodeDelegate::CreateStatic(CustomizeInputNodeLambda, TWeakObjectPtr<const UInputAction>(Action));
+			ActionRegistrar.AddBlueprintAction(Action, NodeSpawner);
 		}
 	}
 	else if (const UInputAction* Action = Cast<const UInputAction>(ActionRegistrar.GetActionKeyFilter()))
