Fix FBindingTarget initialization.

diff --git a/Engine/Plugins/Experimental/EnhancedInput/Source/InputEditor/Private/Tests/InputTestFramework.h b/Engine/Plugins/Experimental/EnhancedInput/Source/InputEditor/Private/Tests/InputTestFramework.h
--- a/Engine/Plugins/Experimental/EnhancedInput/Source/InputEditor/Private/Tests/InputTestFramework.h
+++ b/Engine/Plugins/Experimental/EnhancedInput/Source/InputEditor/Private/Tests/InputTestFramework.h
@@ -49,19 +49,19 @@ struct FBindingTargets
 	}
 
 	UPROPERTY()
-	UInputBindingTarget* Started;
+	UInputBindingTarget* Started = nullptr;
 
 	UPROPERTY()
-	UInputBindingTarget* Ongoing;
+	UInputBindingTarget* Ongoing = nullptr;
 
 	UPROPERTY()
-	UInputBindingTarget* Canceled;
+	UInputBindingTarget* Canceled = nullptr;
 
 	UPROPERTY()
-	UInputBindingTarget* Completed;
+	UInputBindingTarget* Completed = nullptr;
 
 	UPROPERTY()
-	UInputBindingTarget* Triggered;
+	UInputBindingTarget* Triggered = nullptr;
 };
 
 // Mock input subsystems to avoid having to create an actual subsystem + local player + game instance.
