Enhanced Input: added UInputTriggerPulse which enables firing at an interval while an input is actuated

diff --git a/Engine/Plugins/Experimental/EnhancedInput/Source/EnhancedInput/Private/InputTriggers.cpp b/Engine/Plugins/Experimental/EnhancedInput/Source/EnhancedInput/Private/InputTriggers.cpp
--- a/Engine/Plugins/Experimental/EnhancedInput/Source/EnhancedInput/Private/InputTriggers.cpp
+++ b/Engine/Plugins/Experimental/EnhancedInput/Source/EnhancedInput/Private/InputTriggers.cpp
@@ -118,6 +118,38 @@ ETriggerState UInputTriggerTap::UpdateState_Implementation(const UEnhancedPlayer
 	return State;
 }
 
+ETriggerState UInputTriggerPulse::UpdateState_Implementation(const UEnhancedPlayerInput* PlayerInput, FInputActionValue ModifiedValue, float DeltaTime)
+{
+	// Update HeldDuration and derive base state
+	ETriggerState State = Super::UpdateState_Implementation(PlayerInput, ModifiedValue, DeltaTime);
+
+	if (State == ETriggerState::Ongoing)
+	{
+		// If the repeat count limit has not been reached
+		if (TriggerLimit == 0 || TriggerCount < TriggerLimit)
+		{
+			// Trigger when HeldDuration exceeds the interval threshold, optionally trigger on initial actuation
+			if (HeldDuration > (Interval * (bTriggerOnStart ? TriggerCount : TriggerCount + 1)))
+			{
+				++TriggerCount;
+				State = ETriggerState::Triggered;
+			}
+		}
+		else
+		{
+			State = ETriggerState::None;
+		}
+	}
+	else
+	{
+		// Reset repeat count
+		TriggerCount = 0;
+	}
+
+	return State;
+}
+
+
 ETriggerState UInputTriggerChordAction::UpdateState_Implementation(const UEnhancedPlayerInput* PlayerInput, FInputActionValue ModifiedValue, float DeltaTime)
 {
 	// Inherit state from the chorded action
diff --git a/Engine/Plugins/Experimental/EnhancedInput/Source/EnhancedInput/Public/InputTriggers.h b/Engine/Plugins/Experimental/EnhancedInput/Source/EnhancedInput/Public/InputTriggers.h
--- a/Engine/Plugins/Experimental/EnhancedInput/Source/EnhancedInput/Public/InputTriggers.h
+++ b/Engine/Plugins/Experimental/EnhancedInput/Source/EnhancedInput/Public/InputTriggers.h
@@ -265,6 +265,38 @@ public:
 	float TapReleaseTimeThreshold = 0.2f;
 };
 
+/** UInputTriggerPulse
+	Trigger that fires at an Interval, in seconds, while input is actuated. 
+	Note: Completed only fires when the repeat limit is reached or when input is released immediately after being triggered.
+	*/
+UCLASS(NotBlueprintable, MinimalAPI, meta = (DisplayName = "Pulse"))
+class UInputTriggerPulse final : public UInputTriggerTimedBase
+{
+	GENERATED_BODY()
+
+private:
+
+	int32 TriggerCount = 0;
+
+protected:
+
+	virtual ETriggerState UpdateState_Implementation(const UEnhancedPlayerInput* PlayerInput, FInputActionValue ModifiedValue, float DeltaTime) override;
+
+public:
+	// Whether to trigger when the input first exceeds the actuation threshold or wait for the first interval?
+	UPROPERTY(EditAnywhere, Config, BlueprintReadWrite, Category = "Trigger Settings")
+	bool bTriggerOnStart = true;
+
+	// How long between each trigger fire while input is held, in seconds?
+	UPROPERTY(EditAnywhere, Config, BlueprintReadWrite, Category = "Trigger Settings", meta = (ClampMin = "0"))
+	float Interval = 1.0f;
+
+	// How many times can the trigger fire while input is held? (0 = no limit)
+	UPROPERTY(EditAnywhere, Config, BlueprintReadWrite, Category = "Trigger Settings", meta = (ClampMin = "0"))
+	int32 TriggerLimit = 0;
+
+	virtual FString GetDebugState() const override { return HeldDuration ? FString::Printf(TEXT("Triggers:%d/%d, Interval:%.2f/%.2f"), TriggerCount, TriggerLimit, (HeldDuration/(Interval*(TriggerCount+1))), Interval) : FString(); }
+};
 
 
 // Chorded actions
