Enhanced Input: fixed inputs being affected by time dilation and elapsed processing time being reset on Triggered -> Ongoing

diff --git a/Engine/Plugins/Experimental/EnhancedInput/Source/EnhancedInput/Private/EnhancedPlayerInput.cpp b/Engine/Plugins/Experimental/EnhancedInput/Source/EnhancedInput/Private/EnhancedPlayerInput.cpp
--- a/Engine/Plugins/Experimental/EnhancedInput/Source/EnhancedInput/Private/EnhancedPlayerInput.cpp
+++ b/Engine/Plugins/Experimental/EnhancedInput/Source/EnhancedInput/Private/EnhancedPlayerInput.cpp
@@ -28,6 +28,7 @@ ETriggerEventInternal UEnhancedPlayerInput::GetTriggerStateChangeEvent(ETriggerS
 	// Ongoing	 -> Ongoing		= Ongoing
 	// Ongoing	 -> Triggered	= Triggered
 	// Triggered -> Triggered	= Triggered
+	// Triggered -> Ongoing		= Ongoing
 	// Triggered -> None	    = Completed
 
 	switch (LastTriggerState)
@@ -61,6 +62,10 @@ ETriggerEventInternal UEnhancedPlayerInput::GetTriggerStateChangeEvent(ETriggerS
 		{
 			return ETriggerEventInternal::Triggered;	// Don't re-raise Started event for multiple completed ticks.
 		}
+		else if (NewTriggerState == ETriggerState::Ongoing)
+		{
+			return ETriggerEventInternal::Ongoing;
+		}
 		else if (NewTriggerState == ETriggerState::None)
 		{
 			return ETriggerEventInternal::Completed;
@@ -182,6 +187,11 @@ void UEnhancedPlayerInput::ProcessInputStack(const TArray<UInputComponent*>& Inp
 	// Process Action bindings
 	ActionsWithEventsThisTick.Reset();
 
+	// Use non-dilated delta time for processing
+	check(GetOuterAPlayerController());
+	const float Dilation = GetOuterAPlayerController()->GetActorTimeDilation();
+	const float NonDilatedDeltaTime = DeltaTime / Dilation;
+
 	// Handle input devices, applying modifiers and triggers
 	for (FEnhancedActionKeyMapping& Mapping : EnhancedActionMappings)
 	{
@@ -209,7 +219,7 @@ void UEnhancedPlayerInput::ProcessInputStack(const TArray<UInputComponent*>& Inp
 		EKeyEvent KeyEvent = bKeyIsHeld ? EKeyEvent::Held : ((bKeyIsDown || bKeyIsReleased) ? EKeyEvent::Actuated : EKeyEvent::None);
 
 		// Perform update
-		ProcessActionMappingEvent(Mapping.Action, DeltaTime, bGamePaused, RawKeyValue, KeyEvent, Mapping.Modifiers, Mapping.Triggers);
+		ProcessActionMappingEvent(Mapping.Action, NonDilatedDeltaTime, bGamePaused, RawKeyValue, KeyEvent, Mapping.Modifiers, Mapping.Triggers);
 	}
 
 
@@ -226,7 +236,7 @@ void UEnhancedPlayerInput::ProcessInputStack(const TArray<UInputComponent*>& Inp
 		else if (!InputsInjectedThisTick.Contains(InjectedAction))
 		{
 			// Reset action state by "releasing the key".
-			ProcessActionMappingEvent(InjectedAction, DeltaTime, bGamePaused, FInputActionValue(), EKeyEvent::Actuated, {}, {});
+			ProcessActionMappingEvent(InjectedAction, NonDilatedDeltaTime, bGamePaused, FInputActionValue(), EKeyEvent::Actuated, {}, {});
 			It.RemoveCurrent();
 		}
 	}
@@ -249,7 +259,7 @@ void UEnhancedPlayerInput::ProcessInputStack(const TArray<UInputComponent*>& Inp
 		for (FInjectedInput& InjectedInput : InjectedPair.Value.Injected)
 		{
 			// Perform update
-			ProcessActionMappingEvent(InjectedAction, DeltaTime, bGamePaused, InjectedInput.RawValue, KeyEvent, InjectedInput.Modifiers, InjectedInput.Triggers);
+			ProcessActionMappingEvent(InjectedAction, NonDilatedDeltaTime, bGamePaused, InjectedInput.RawValue, KeyEvent, InjectedInput.Modifiers, InjectedInput.Triggers);
 		}
 	}
 	InputsInjectedThisTick.Reset();
@@ -267,10 +277,10 @@ void UEnhancedPlayerInput::ProcessInputStack(const TArray<UInputComponent*>& Inp
 		if (ActionsWithEventsThisTick.Contains(Action))
 		{
 			// Apply modifiers
-			ActionData.Value = ApplyModifiers(ActionData.Modifiers, ActionData.Value, DeltaTime);
+			ActionData.Value = ApplyModifiers(ActionData.Modifiers, ActionData.Value, NonDilatedDeltaTime);
 
 			// Evaluate triggers
-			TriggerState = CalcTriggerState(ActionData.Triggers, ActionData.Value, DeltaTime);
+			TriggerState = CalcTriggerState(ActionData.Triggers, ActionData.Value, NonDilatedDeltaTime);
 
 			// Any mapping triggers applied should limit the final state.
 			TriggerState = ActionData.bMappingTriggerApplied ? FMath::Min(TriggerState, ActionData.MappingTriggerState) : TriggerState;
@@ -288,8 +298,8 @@ void UEnhancedPlayerInput::ProcessInputStack(const TArray<UInputComponent*>& Inp
 		ActionData.TriggerEvent = ConvertInternalTriggerEvent(ActionData.TriggerEventInternal);
 		ActionData.LastTriggerState = TriggerState;
 		// Evaluate time per action after establishing the internal trigger state across all mappings
-		ActionData.ElapsedProcessedTime += TriggerState != ETriggerState::None ? DeltaTime : 0.f;
-		ActionData.ElapsedTriggeredTime += (ActionData.TriggerEvent == ETriggerEvent::Triggered) ? DeltaTime : 0.f;
+		ActionData.ElapsedProcessedTime += TriggerState != ETriggerState::None ? NonDilatedDeltaTime : 0.f;
+		ActionData.ElapsedTriggeredTime += (ActionData.TriggerEvent == ETriggerEvent::Triggered) ? NonDilatedDeltaTime : 0.f;
 	}
 
 
diff --git a/Engine/Plugins/Experimental/EnhancedInput/Source/EnhancedInput/Private/InputTriggers.cpp b/Engine/Plugins/Experimental/EnhancedInput/Source/EnhancedInput/Private/InputTriggers.cpp
--- a/Engine/Plugins/Experimental/EnhancedInput/Source/EnhancedInput/Private/InputTriggers.cpp
+++ b/Engine/Plugins/Experimental/EnhancedInput/Source/EnhancedInput/Private/InputTriggers.cpp
@@ -16,11 +16,11 @@ ETriggerState UInputTriggerTimedBase::UpdateState_Implementation(const UEnhanced
 {
 	ETriggerState State = ETriggerState::None;
 
-	// Transition to Ongoing on actuation.
+	// Transition to Ongoing on actuation. Update the held duration.
 	if (IsActuated(ModifiedValue))
 	{
 		State = ETriggerState::Ongoing;
-		HeldDuration += DeltaTime;	// TODO: When attached directly to an Action this will tick N times a frame where N is the number of evaluated (actively held) mappings.
+		HeldDuration = CalculateHeldDuration(PlayerInput, DeltaTime);
 	}
 	else
 	{
@@ -31,6 +31,15 @@ ETriggerState UInputTriggerTimedBase::UpdateState_Implementation(const UEnhanced
 	return State;
 }
 
+float UInputTriggerTimedBase::CalculateHeldDuration(const UEnhancedPlayerInput* const PlayerInput, const float DeltaTime) const
+{
+	check(PlayerInput && PlayerInput->GetOuterAPlayerController());
+	const float TimeDilation = PlayerInput->GetOuterAPlayerController()->GetActorTimeDilation();
+	
+	// Calculates the new held duration, applying time dilation if desired
+	return HeldDuration + (!bAffectedByTimeDilation ? DeltaTime : DeltaTime * TimeDilation);
+}
+
 
 // Implementations
 
@@ -83,7 +92,7 @@ ETriggerState UInputTriggerHoldAndRelease::UpdateState_Implementation(const UEnh
 {
 	// Evaluate the updated held duration prior to calling Super to update the held timer
 	// This stops us failing to trigger if the input is released on the threshold frame due to HeldDuration being 0.
-	float TickHeldDuration = HeldDuration + DeltaTime;
+	const float TickHeldDuration = CalculateHeldDuration(PlayerInput, DeltaTime);
 
 	// Update HeldDuration and derive base state
 	ETriggerState State = Super::UpdateState_Implementation(PlayerInput, ModifiedValue, DeltaTime);
diff --git a/Engine/Plugins/Experimental/EnhancedInput/Source/EnhancedInput/Public/InputTriggers.h b/Engine/Plugins/Experimental/EnhancedInput/Source/EnhancedInput/Public/InputTriggers.h
--- a/Engine/Plugins/Experimental/EnhancedInput/Source/EnhancedInput/Public/InputTriggers.h
+++ b/Engine/Plugins/Experimental/EnhancedInput/Source/EnhancedInput/Public/InputTriggers.h
@@ -142,8 +142,15 @@ protected:
 	// Transitions to Ongoing on actuation. Never triggers.
 	virtual ETriggerState UpdateState_Implementation(const UEnhancedPlayerInput* PlayerInput, FInputActionValue ModifiedValue, float DeltaTime) override;
 
+	// Calculates the new held duration given the current player input and delta time
+	float CalculateHeldDuration(const UEnhancedPlayerInput* const PlayerInput, const float DeltaTime) const;
+
 public:
 
+	// Should global time dilation be applied to the held duration?
+	UPROPERTY(BlueprintReadWrite, Category = "Trigger Settings")
+	bool bAffectedByTimeDilation = false;
+
 	virtual FString GetDebugState() const override { return HeldDuration ? FString::Printf(TEXT("Held:%.2f"), HeldDuration) : FString(); }
 };
 
@@ -267,7 +274,8 @@ public:
 
 /** UInputTriggerPulse
 	Trigger that fires at an Interval, in seconds, while input is actuated. 
-	Note: Completed only fires when the repeat limit is reached or when input is released immediately after being triggered.
+	Note:	Completed only fires when the repeat limit is reached or when input is released immediately after being triggered.
+			Otherwise, Canceled is fired when input is released.
 	*/
 UCLASS(NotBlueprintable, MinimalAPI, meta = (DisplayName = "Pulse"))
 class UInputTriggerPulse final : public UInputTriggerTimedBase
