Moving a details panel creation to PostEngineInit

diff --git a/Engine/Plugins/Experimental/EnhancedInput/Source/InputEditor/Private/InputEditorModule.cpp b/Engine/Plugins/Experimental/EnhancedInput/Source/InputEditor/Private/InputEditorModule.cpp
--- a/Engine/Plugins/Experimental/EnhancedInput/Source/InputEditor/Private/InputEditorModule.cpp
+++ b/Engine/Plugins/Experimental/EnhancedInput/Source/InputEditor/Private/InputEditorModule.cpp
@@ -53,7 +53,7 @@ private:
 		AssetTools.RegisterAssetTypeActions(Action);
 		CreatedAssetTypeActions.Add(Action);
 	}
-
+	void PostEngineInit();
 	TSharedRef<SWidget> CreateSettingsPanel();
 	void OnSettingChanged(const FPropertyChangedEvent& PropertyChangedEvent);
 
@@ -367,6 +367,20 @@ TArray<UObject*> FInputEditorModule::GatherClassDetailsCDOs(UClass* Class, const
 	return CDOs;
 }
 
+void FInputEditorModule::PostEngineInit()
+{
+	// Register input settings
+	ISettingsModule* SettingsModule = FModuleManager::GetModulePtr<ISettingsModule>("Settings");
+	if (SettingsModule && FSlateApplication::IsInitialized())
+	{
+		SettingsModule->RegisterSettings("Project", "Plugins", "EnhancedInput",
+			LOCTEXT("EnhancedInputSettingsName", "Enhanced Input"),
+			LOCTEXT("EnhancedInputSettingsDescription", "Modify defaults for configurable triggers and modifiers."),
+			CreateSettingsPanel()
+		);
+	}
+}
+
 TSharedRef<SWidget> FInputEditorModule::CreateSettingsPanel()
 {
 	TSharedPtr<IDetailsView> TriggerDetailsView = AddClassDetailsView<UInputTrigger>();
@@ -415,17 +429,6 @@ void FInputEditorModule::StartupModule()
 	PropertyModule.RegisterCustomPropertyTypeLayout("EnhancedActionKeyMapping", FOnGetPropertyTypeCustomizationInstance::CreateStatic(&FEnhancedActionMappingCustomization::MakeInstance));
 	PropertyModule.NotifyCustomizationModuleChanged();
 
-	// Register input settings
-	ISettingsModule* SettingsModule = FModuleManager::GetModulePtr<ISettingsModule>("Settings");
-	if (SettingsModule && FSlateApplication::IsInitialized())
-	{
-		SettingsModule->RegisterSettings("Project", "Plugins", "EnhancedInput",
-			LOCTEXT("EnhancedInputSettingsName", "Enhanced Input"),
-			LOCTEXT("EnhancedInputSettingsDescription", "Modify defaults for configurable triggers and modifiers."),
-			CreateSettingsPanel()
-		);
-	}
-
 	// Register input assets
 	IAssetTools& AssetTools = FModuleManager::LoadModuleChecked<FAssetToolsModule>("AssetTools").Get();
 	InputAssetsCategory = AssetTools.RegisterAdvancedAssetCategory(FName(TEXT("Input")), LOCTEXT("InputAssetsCategory", "Input"));
@@ -444,6 +447,8 @@ void FInputEditorModule::StartupModule()
 	AssetRegistry.OnAssetRemoved().AddRaw(this, &FInputEditorModule::OnAssetRemoved);
 	AssetRegistry.OnAssetRenamed().AddRaw(this, &FInputEditorModule::OnAssetRenamed);
 	// TODO: Update settings whenever a config variable is added/removed to an asset
+
+	FCoreDelegates::OnPostEngineInit.AddRaw(this, &FInputEditorModule::PostEngineInit);
 }
 
 void FInputEditorModule::ShutdownModule()
@@ -477,6 +482,8 @@ void FInputEditorModule::ShutdownModule()
 	PropertyModule.UnregisterCustomClassLayout("InputContext");
 	PropertyModule.UnregisterCustomPropertyTypeLayout("EnhancedActionKeyMapping");
 	PropertyModule.NotifyCustomizationModuleChanged();
+
+	FCoreDelegates::OnPostEngineInit.RemoveAll(this);
 }
 
 FInputEditorModule::FClassDetailsView FInputEditorModule::FindClassDetailsViewForAsset(const FAssetData& AssetData)
