Lyra - Working on input binding support, checkpointing progress we now have defaults that are being converted into input contexts at runtime, and in theory can be modified.

diff --git a/Engine/Plugins/Experimental/EnhancedInput/Source/EnhancedInput/Public/InputMappingContext.h b/Engine/Plugins/Experimental/EnhancedInput/Source/EnhancedInput/Public/InputMappingContext.h
--- a/Engine/Plugins/Experimental/EnhancedInput/Source/EnhancedInput/Public/InputMappingContext.h
+++ b/Engine/Plugins/Experimental/EnhancedInput/Source/EnhancedInput/Public/InputMappingContext.h
@@ -23,18 +23,20 @@ class ENHANCEDINPUT_API UInputMappingContext : public UDataAsset
 {
 	GENERATED_BODY()
 
+public:
+
+	// Localized context descriptor
+	UPROPERTY(BlueprintReadOnly, EditAnywhere, Category = "Config", DisplayName = "Description")
+	FText ContextDescription;
+
 protected:
 	// List of key to action mappings.
 	UPROPERTY(config, BlueprintReadOnly, EditAnywhere, Category = "Mappings")
 	TArray<FEnhancedActionKeyMapping> Mappings;
 
 	friend class FInputContextDetails;
-public:
-
-	// Localized context descriptor
-	UPROPERTY(BlueprintReadOnly, EditAnywhere, Category = "Config", DisplayName = "Description")
-	FText ContextDescription;
 
+public:
 
 	/**
 	* Mapping accessors.
