Expose the source action of an FInputActionInstance to native and blueprints.

#rnx

diff --git a/Engine/Plugins/Experimental/EnhancedInput/Source/EnhancedInput/Public/InputAction.h b/Engine/Plugins/Experimental/EnhancedInput/Source/EnhancedInput/Public/InputAction.h
--- a/Engine/Plugins/Experimental/EnhancedInput/Source/EnhancedInput/Public/InputAction.h
+++ b/Engine/Plugins/Experimental/EnhancedInput/Source/EnhancedInput/Public/InputAction.h
@@ -71,7 +71,9 @@ struct ENHANCEDINPUT_API FInputActionInstance
 	GENERATED_BODY()
 
 private:
-	UPROPERTY()
+
+	// The source action that this instance is created from
+	UPROPERTY(BlueprintReadOnly, Category = "Input", Meta = (AllowPrivateAccess))
 	const UInputAction* SourceAction = nullptr;
 
 	// Internal trigger states
@@ -130,4 +132,7 @@ public:
 
 	UE_DEPRECATED(4.26, "GetModifiers(EModifierExecutionPhase) is deprecated. Use GetModifiers()")
 	const TArray<UInputModifier*>& GetModifiers(EModifierExecutionPhase ForPhase) const { return Modifiers; }
+
+	// The source action that this instance is created from
+	const UInputAction* GetSourceAction() const { return SourceAction; }
 };
\ No newline at end of file
