EnhancedInput - Making the InputContext and InputAction asset actions, derive from the DataAsset version that's now exposed so that they get the change baseclass option.

diff --git a/Engine/Plugins/Experimental/EnhancedInput/Source/InputEditor/InputEditor.Build.cs b/Engine/Plugins/Experimental/EnhancedInput/Source/InputEditor/InputEditor.Build.cs
--- a/Engine/Plugins/Experimental/EnhancedInput/Source/InputEditor/InputEditor.Build.cs
+++ b/Engine/Plugins/Experimental/EnhancedInput/Source/InputEditor/InputEditor.Build.cs
@@ -35,6 +35,7 @@ namespace UnrealBuildTool.Rules
                     "Slate",
                     "SlateCore",
                     "UnrealEd",
+					"AssetTools"
                 }
             );
         }
diff --git a/Engine/Plugins/Experimental/EnhancedInput/Source/InputEditor/Private/InputEditorModule.cpp b/Engine/Plugins/Experimental/EnhancedInput/Source/InputEditor/Private/InputEditorModule.cpp
--- a/Engine/Plugins/Experimental/EnhancedInput/Source/InputEditor/Private/InputEditorModule.cpp
+++ b/Engine/Plugins/Experimental/EnhancedInput/Source/InputEditor/Private/InputEditorModule.cpp
@@ -28,6 +28,7 @@
 #include "TickableEditorObject.h"
 #include "Widgets/Layout/SScrollBox.h"
 #include "Widgets/Notifications/SNotificationList.h"
+#include "AssetTypeActions/AssetTypeActions_DataAsset.h"
 
 #define LOCTEXT_NAMESPACE "InputEditor"
 
@@ -141,7 +142,7 @@ UObject* UInputAction_Factory::FactoryCreateNew(UClass* Class, UObject* InParent
 // Asset type actions
 // TODO: Move asset type action definitions out?
 
-class FAssetTypeActions_InputContext : public FAssetTypeActions_Base {
+class FAssetTypeActions_InputContext : public FAssetTypeActions_DataAsset {
 public:
 	virtual FText GetName() const override { return NSLOCTEXT("AssetTypeActions", "AssetTypeActions_InputMappingContext", "Input Mapping Context"); }
 	virtual uint32 GetCategories() override { return FInputEditorModule::GetInputAssetsCategory(); }
@@ -150,7 +151,7 @@ public:
 	virtual UClass* GetSupportedClass() const override { return UInputMappingContext::StaticClass(); }
 };
 
-class FAssetTypeActions_InputAction : public FAssetTypeActions_Base {
+class FAssetTypeActions_InputAction : public FAssetTypeActions_DataAsset {
 public:
 	virtual FText GetName() const override { return NSLOCTEXT("AssetTypeActions", "AssetTypeActions_InputAction", "Input Action"); }
 	virtual uint32 GetCategories() override { return FInputEditorModule::GetInputAssetsCategory(); }
@@ -159,23 +160,8 @@ public:
 	virtual UClass* GetSupportedClass() const override { return UInputAction::StaticClass(); }
 };
 
-//class FAssetTypeActions_InputTrigger : public FAssetTypeActions_Base {
-//public:
-//	virtual FText GetName() const override { return NSLOCTEXT("AssetTypeActions", "AssetTypeActions_InputTrigger", "Input Trigger"); }
-//	virtual uint32 GetCategories() override { return FInputEditorModule::GetInputAssetsCategory(); }
-//	virtual FColor GetTypeColor() const override { return FColor(127, 0, 255); }
-//	virtual FText GetAssetDescription(const FAssetData& AssetData) const override { return NSLOCTEXT("AssetTypeActions", "AssetTypeActions_InputTriggerDesc", "Supports custom trigger rules for actions and device inputs."); }
-//	virtual UClass* GetSupportedClass() const override { return UInputTrigger::StaticClass(); }
-//};
-//
-//class FAssetTypeActions_InputModifier : public FAssetTypeActions_Base {
-//public:
-//	virtual FText GetName() const override { return NSLOCTEXT("AssetTypeActions", "AssetTypeActions_InputModifier", "Input Modifier"); }
-//	virtual uint32 GetCategories() override { return FInputEditorModule::GetInputAssetsCategory(); }
-//	virtual FColor GetTypeColor() const override { return FColor(127, 255, 0); }
-//	virtual FText GetAssetDescription(const FAssetData& AssetData) const override { return NSLOCTEXT("AssetTypeActions", "AssetTypeActions_InputModifierDesc", "Applied to actions and raw device inputs to modify their output."); }
-//	virtual UClass* GetSupportedClass() const override { return UInputModifier::StaticClass(); }
-//};
+
+
 
 void FInputEditorModule::OnSettingChanged(const FPropertyChangedEvent& PropertyChangedEvent)
 {
