Add the input action as an output pin to K2Node_EnhancedInputActions

#rnx

diff --git a/Engine/Plugins/Experimental/EnhancedInput/Source/EnhancedInput/Public/EnhancedInputComponent.h b/Engine/Plugins/Experimental/EnhancedInput/Source/EnhancedInput/Public/EnhancedInputComponent.h
--- a/Engine/Plugins/Experimental/EnhancedInput/Source/EnhancedInput/Public/EnhancedInputComponent.h
+++ b/Engine/Plugins/Experimental/EnhancedInput/Source/EnhancedInput/Public/EnhancedInputComponent.h
@@ -24,7 +24,7 @@ DECLARE_DYNAMIC_DELEGATE_OneParam(FInputDebugKeyHandlerDynamicSignature, FKey, K
 DECLARE_DELEGATE(FEnhancedInputActionHandlerSignature);
 DECLARE_DELEGATE_OneParam(FEnhancedInputActionHandlerValueSignature, const FInputActionValue&);
 DECLARE_DELEGATE_OneParam(FEnhancedInputActionHandlerInstanceSignature, const FInputActionInstance&);	// Provides full access to value and timers
-DECLARE_DYNAMIC_DELEGATE_ThreeParams(FEnhancedInputActionHandlerDynamicSignature, FInputActionValue, ActionValue, float, ElapsedTime, float, TriggeredTime);
+DECLARE_DYNAMIC_DELEGATE_FourParams(FEnhancedInputActionHandlerDynamicSignature, FInputActionValue, ActionValue, float, ElapsedTime, float, TriggeredTime, const UInputAction*, SourceAction);
 
 /** Unified storage for both native and dynamic delegates with any signature  */
 template<typename TSignature>
@@ -266,7 +266,7 @@ inline void FEnhancedInputActionEventDelegateBinding<FEnhancedInputActionHandler
 template<>
 inline void FEnhancedInputActionEventDelegateBinding<FEnhancedInputActionHandlerDynamicSignature>::Execute(const FInputActionInstance& ActionData) const
 {
-	Delegate.Execute(ActionData.GetValue(), ActionData.GetElapsedTime(), ActionData.GetTriggeredTime());
+	Delegate.Execute(ActionData.GetValue(), ActionData.GetElapsedTime(), ActionData.GetTriggeredTime(), ActionData.GetSourceAction());
 }
 
 
diff --git a/Engine/Plugins/Experimental/EnhancedInput/Source/InputBlueprintNodes/Private/K2Node_EnhancedInputAction.cpp b/Engine/Plugins/Experimental/EnhancedInput/Source/InputBlueprintNodes/Private/K2Node_EnhancedInputAction.cpp
--- a/Engine/Plugins/Experimental/EnhancedInput/Source/InputBlueprintNodes/Private/K2Node_EnhancedInputAction.cpp
+++ b/Engine/Plugins/Experimental/EnhancedInput/Source/InputBlueprintNodes/Private/K2Node_EnhancedInputAction.cpp
@@ -23,6 +23,11 @@
 
 #define LOCTEXT_NAMESPACE "K2Node_EnhancedInputAction"
 
+static const FName InputActionPinName = TEXT("InputAction");
+static const FName ElapsedSecondsPinName = TEXT("ElapsedSeconds");
+static const FName TriggeredSecondsPinName = TEXT("TriggeredSeconds");
+static const FName ActionValuePinName = TEXT("ActionValue");
+
 UK2Node_EnhancedInputAction::UK2Node_EnhancedInputAction(const FObjectInitializer& ObjectInitializer)
 	: Super(ObjectInitializer)
 {
@@ -46,15 +51,22 @@ void UK2Node_EnhancedInputAction::AllocateDefaultPins()
 
 	ForEachEventPinName([this](ETriggerEvent Event, FName PinName) { CreatePin(EGPD_Output, UEdGraphSchema_K2::PC_Exec, PinName); });
 	HideEventPins(nullptr);
-
+	const UEdGraphSchema_K2* Schema = GetDefault<UEdGraphSchema_K2>();
+	
 	AdvancedPinDisplay = ENodeAdvancedPins::Hidden;
 
-	UEdGraphPin* ValuePin = CreatePin(EGPD_Output, UK2Node_GetInputActionValue::GetValueCategory(InputAction), UK2Node_GetInputActionValue::GetValueSubCategoryObject(InputAction), TEXT("ActionValue"));
-	GetDefault<UEdGraphSchema_K2>()->SetPinAutogeneratedDefaultValueBasedOnType(ValuePin);
+	UEdGraphPin* ValuePin = CreatePin(EGPD_Output, UK2Node_GetInputActionValue::GetValueCategory(InputAction), UK2Node_GetInputActionValue::GetValueSubCategoryObject(InputAction), ActionValuePinName);
+	
+	Schema->SetPinAutogeneratedDefaultValueBasedOnType(ValuePin);
 
-	CreatePin(EGPD_Output, UEdGraphSchema_K2::PC_Float, TEXT("ElapsedSeconds"))->bAdvancedView = true;
-	CreatePin(EGPD_Output, UEdGraphSchema_K2::PC_Float, TEXT("TriggeredSeconds"))->bAdvancedView = true;
+	CreatePin(EGPD_Output, UEdGraphSchema_K2::PC_Float, ElapsedSecondsPinName)->bAdvancedView = true;
+	CreatePin(EGPD_Output, UEdGraphSchema_K2::PC_Float, TriggeredSecondsPinName)->bAdvancedView = true;
 
+	UEdGraphPin* ActionPin = CreatePin(EGPD_Output, UEdGraphSchema_K2::PC_Object, InputAction->GetClass(), InputActionPinName);
+	ActionPin->DefaultObject = const_cast<UObject*>(Cast<UObject>(InputAction));
+	ActionPin->DefaultValue = InputAction->GetName();
+	Schema->ConstructBasicPinTooltip(*ActionPin, LOCTEXT("InputActionPinDescription", "The input action that caused this event to fire"), ActionPin->PinToolTip);
+	
 	Super::AllocateDefaultPins();
 }
 
@@ -227,6 +239,10 @@ void UK2Node_EnhancedInputAction::ExpandNode(FKismetCompilerContext& CompilerCon
 	UK2Node_TemporaryVariable* TriggeredSecondsVar = CompilerContext.SpawnIntermediateNode<UK2Node_TemporaryVariable>(this, SourceGraph);
 	TriggeredSecondsVar->VariableType.PinCategory = UEdGraphSchema_K2::PC_Float;
 	TriggeredSecondsVar->AllocateDefaultPins();
+	UK2Node_TemporaryVariable* InputActionVar = CompilerContext.SpawnIntermediateNode<UK2Node_TemporaryVariable>(this, SourceGraph);
+	InputActionVar->VariableType.PinCategory = UEdGraphSchema_K2::PC_Object;
+	InputActionVar->VariableType.PinSubCategoryObject = InputAction->GetClass();
+	InputActionVar->AllocateDefaultPins();
 
 	for (ActivePinData& PinData : ActivePins)
 	{
@@ -242,11 +258,11 @@ void UK2Node_EnhancedInputAction::ExpandNode(FKismetCompilerContext& CompilerCon
 		UK2Node_AssignmentStatement* ActionValueInitialize = CompilerContext.SpawnIntermediateNode<UK2Node_AssignmentStatement>(this, SourceGraph);
 		ActionValueInitialize->AllocateDefaultPins();
 		Schema->TryCreateConnection(ActionValueVar->GetVariablePin(), ActionValueInitialize->GetVariablePin());
-		Schema->TryCreateConnection(ActionValueInitialize->GetValuePin(), InputActionEvent->FindPinChecked(TEXT("ActionValue")));
+		Schema->TryCreateConnection(ActionValueInitialize->GetValuePin(), InputActionEvent->FindPinChecked(ActionValuePinName));
 		// Connect the events to the assign location nodes
 		Schema->TryCreateConnection(Schema->FindExecutionPin(*InputActionEvent, EGPD_Output), ActionValueInitialize->GetExecPin());
 
-		// Create assignment nodes to assign the elapsed timers
+		// Create assignment nodes to assign the elapsed timers and input action
 		UK2Node_AssignmentStatement* ElapsedSecondsInitialize = CompilerContext.SpawnIntermediateNode<UK2Node_AssignmentStatement>(this, SourceGraph);
 		ElapsedSecondsInitialize->AllocateDefaultPins();
 		Schema->TryCreateConnection(ElapsedSecondsVar->GetVariablePin(), ElapsedSecondsInitialize->GetVariablePin());
@@ -257,25 +273,27 @@ void UK2Node_EnhancedInputAction::ExpandNode(FKismetCompilerContext& CompilerCon
 		Schema->TryCreateConnection(TriggeredSecondsVar->GetVariablePin(), TriggeredSecondsInitialize->GetVariablePin());
 		Schema->TryCreateConnection(TriggeredSecondsInitialize->GetValuePin(), InputActionEvent->FindPinChecked(TEXT("TriggeredTime")));
 
+		UK2Node_AssignmentStatement* InputActionInitialize = CompilerContext.SpawnIntermediateNode<UK2Node_AssignmentStatement>(this, SourceGraph);
+		InputActionInitialize->AllocateDefaultPins();
+		Schema->TryCreateConnection(InputActionVar->GetVariablePin(), InputActionInitialize->GetVariablePin());
+		Schema->TryCreateConnection(InputActionInitialize->GetValuePin(), InputActionEvent->FindPinChecked(TEXT("SourceAction")));
+		 
 		// Connect the assign location to the assign elapsed time nodes
 		Schema->TryCreateConnection(ActionValueInitialize->GetThenPin(), ElapsedSecondsInitialize->GetExecPin());
 		Schema->TryCreateConnection(ElapsedSecondsInitialize->GetThenPin(), TriggeredSecondsInitialize->GetExecPin());
-
-		// Move the original event connections to the then pin of the elapsed time assign
-		CompilerContext.MovePinLinksToIntermediate(*EachPin, *TriggeredSecondsInitialize->GetThenPin());
+		Schema->TryCreateConnection(TriggeredSecondsInitialize->GetThenPin(), InputActionInitialize->GetExecPin());
+		
+		// Move the original event connections to the then pin of the Input Action assign
+		CompilerContext.MovePinLinksToIntermediate(*EachPin, *InputActionInitialize->GetThenPin());
 
 		// Move the original event variable connections to the intermediate nodes
-		CompilerContext.MovePinLinksToIntermediate(*FindPin(TEXT("ActionValue")), *ActionValueVar->GetVariablePin());
-
-		CompilerContext.MovePinLinksToIntermediate(*FindPin(TEXT("ElapsedSeconds")), *ElapsedSecondsVar->GetVariablePin());
-		CompilerContext.MovePinLinksToIntermediate(*FindPin(TEXT("TriggeredSeconds")), *TriggeredSecondsVar->GetVariablePin());
+		CompilerContext.MovePinLinksToIntermediate(*FindPin(ActionValuePinName), *ActionValueVar->GetVariablePin());
+		CompilerContext.MovePinLinksToIntermediate(*FindPin(ElapsedSecondsPinName), *ElapsedSecondsVar->GetVariablePin());
+		CompilerContext.MovePinLinksToIntermediate(*FindPin(TriggeredSecondsPinName), *TriggeredSecondsVar->GetVariablePin());
+		CompilerContext.MovePinLinksToIntermediate(*FindPin(InputActionPinName), *InputActionVar->GetVariablePin());
 	}
 }
 
-void PopulateMenuActions(FBlueprintActionDatabaseRegistrar& ActionRegistrar)
-{
-}
-
 void UK2Node_EnhancedInputAction::GetMenuActions(FBlueprintActionDatabaseRegistrar& ActionRegistrar) const
 {
 	auto CustomizeInputNodeLambda = [](UEdGraphNode* NewNode, bool bIsTemplateNode, TWeakObjectPtr<const UInputAction> Action)
