Make the UEnhancedInputEngineSubsystem a Non-Blueprint type as it is experimental and should not yet be exposed to BP

#rnx

diff --git a/Engine/Plugins/Experimental/EnhancedInput/Source/EnhancedInput/Public/EnhancedInputSubsystems.h b/Engine/Plugins/Experimental/EnhancedInput/Source/EnhancedInput/Public/EnhancedInputSubsystems.h
--- a/Engine/Plugins/Experimental/EnhancedInput/Source/EnhancedInput/Public/EnhancedInputSubsystems.h
+++ b/Engine/Plugins/Experimental/EnhancedInput/Source/EnhancedInput/Public/EnhancedInputSubsystems.h
@@ -23,8 +23,8 @@ public:
 };
 
 // Global input handling subsystem
-// TODO: For now this is a non-functional placeholder.
-UCLASS()
+// TODO: For now this is a non-functional placeholder that shuold be enabled when the Enhanced Input Plugin is enabled by default.
+UCLASS(NotBlueprintType)
 class ENHANCEDINPUT_API UEnhancedInputEngineSubsystem : public UEngineSubsystem, public IEnhancedInputSubsystemInterface
 {
 	GENERATED_BODY()
