[Enhanced Input] Add the current key value to Debug Nodes by passing in the current Key Map value of an FKey.

#rnx

diff --git a/Engine/Plugins/Experimental/EnhancedInput/Source/EnhancedInput/Private/EnhancedPlayerInput.cpp b/Engine/Plugins/Experimental/EnhancedInput/Source/EnhancedInput/Private/EnhancedPlayerInput.cpp
--- a/Engine/Plugins/Experimental/EnhancedInput/Source/EnhancedInput/Private/EnhancedPlayerInput.cpp
+++ b/Engine/Plugins/Experimental/EnhancedInput/Source/EnhancedInput/Private/EnhancedPlayerInput.cpp
@@ -384,7 +384,8 @@ void UEnhancedPlayerInput::ProcessInputStack(const TArray<UInputComponent*>& Inp
 				if (!bGamePaused || KeyBinding->bExecuteWhenPaused)
 				{
 					FKeyState* KeyState = KeyStateMap.Find(KeyBinding->Chord.Key);
-					if (KeyState && KeyState->EventCounts[KeyBinding->KeyEvent].Num() > 0)
+					// We always want to update any analog debug events, like Gamepad axis 
+					if ((KeyState && KeyState->EventCounts[KeyBinding->KeyEvent].Num() > 0) || KeyBinding->Chord.Key.IsAnalog())
 					{
 						// Record intent to trigger
 						TriggeredDebugDelegates.Add(KeyBinding->Clone());
@@ -396,7 +397,11 @@ void UEnhancedPlayerInput::ProcessInputStack(const TArray<UInputComponent*>& Inp
 		// Action all debug delegates that triggered this tick, in the order in which they triggered.
 		for (TUniquePtr<FInputDebugKeyBinding>& Delegate : TriggeredDebugDelegates)
 		{
-			Delegate->Execute();
+			const FKeyState* KeyState = GetKeyState(Delegate->Chord.Key);
+			
+			FInputActionValue ActionValue(KeyState ? KeyState->RawValue : FVector::ZeroVector);
+			
+			Delegate->Execute(ActionValue);
 		}
 		TriggeredDebugDelegates.Reset();
 #endif
diff --git a/Engine/Plugins/Experimental/EnhancedInput/Source/EnhancedInput/Public/EnhancedInputComponent.h b/Engine/Plugins/Experimental/EnhancedInput/Source/EnhancedInput/Public/EnhancedInputComponent.h
--- a/Engine/Plugins/Experimental/EnhancedInput/Source/EnhancedInput/Public/EnhancedInputComponent.h
+++ b/Engine/Plugins/Experimental/EnhancedInput/Source/EnhancedInput/Public/EnhancedInputComponent.h
@@ -17,8 +17,8 @@
 enum class ETriggerEvent : uint8;
 
 /** Delegate signature for debug key events. */
-DECLARE_DELEGATE_OneParam(FInputDebugKeyHandlerSignature, FKey);
-DECLARE_DYNAMIC_DELEGATE_OneParam(FInputDebugKeyHandlerDynamicSignature, FKey, Key);
+DECLARE_DELEGATE_TwoParams(FInputDebugKeyHandlerSignature, FKey, FInputActionValue);
+DECLARE_DYNAMIC_DELEGATE_TwoParams(FInputDebugKeyHandlerDynamicSignature, FKey, Key, FInputActionValue, ActionValue);
 
 /** Delegate signature for action events. */
 DECLARE_DELEGATE(FEnhancedInputActionHandlerSignature);
@@ -191,11 +191,9 @@ public:
 		, Chord(InChord)
 	{ }
 
-	virtual void Execute() const = 0;
+	virtual void Execute(const FInputActionValue& ActionValue) const = 0;
 	virtual TUniquePtr<FInputDebugKeyBinding> Clone() const = 0;
 };
-// TODO:: Add FInputDebugKeyValueBinding?
-
 
 
 /**
@@ -230,9 +228,9 @@ private:
 public:
 	FInputDebugKeyDelegateBinding(const FInputChord Chord, const EInputEvent KeyEvent, bool bExecuteWhenPaused) : FInputDebugKeyBinding(Chord, KeyEvent, bExecuteWhenPaused) {}
 
-	virtual void Execute() const override
+	virtual void Execute(const FInputActionValue& ActionValue) const override
 	{
-		Delegate.Execute(Chord.Key);	// TODO: Remove FKey param? We don't support AnyKey, so it isn't terribly useful. TODO: Provide key raw value param to support axis keys.
+		Delegate.Execute(Chord.Key, ActionValue);	// TODO: Remove FKey param? We don't support AnyKey, so it isn't terribly useful.
 	}
 	virtual TUniquePtr<FInputDebugKeyBinding> Clone() const override 
 	{
diff --git a/Engine/Plugins/Experimental/EnhancedInput/Source/InputBlueprintNodes/Private/K2Node_InputDebugKey.cpp b/Engine/Plugins/Experimental/EnhancedInput/Source/InputBlueprintNodes/Private/K2Node_InputDebugKey.cpp
--- a/Engine/Plugins/Experimental/EnhancedInput/Source/InputBlueprintNodes/Private/K2Node_InputDebugKey.cpp
+++ b/Engine/Plugins/Experimental/EnhancedInput/Source/InputBlueprintNodes/Private/K2Node_InputDebugKey.cpp
@@ -12,9 +12,12 @@
 #include "BlueprintNodeSpawner.h"
 #include "EditorCategoryUtils.h"
 #include "BlueprintActionDatabaseRegistrar.h"
+#include "InputActionValue.h"
 
 #define LOCTEXT_NAMESPACE "UK2Node_InputDebugKey"
 
+static const FName ActionPinName = TEXT("ActionValue");
+
 UK2Node_InputDebugKey::UK2Node_InputDebugKey(const FObjectInitializer& ObjectInitializer)
 	: Super(ObjectInitializer)
 {
@@ -46,7 +49,8 @@ void UK2Node_InputDebugKey::AllocateDefaultPins()
 	CreatePin(EGPD_Output, UEdGraphSchema_K2::PC_Exec, TEXT("Pressed"));
 	CreatePin(EGPD_Output, UEdGraphSchema_K2::PC_Exec, TEXT("Released"));
 	CreatePin(EGPD_Output, UEdGraphSchema_K2::PC_Struct, FKey::StaticStruct(), TEXT("Key"));
-
+	CreatePin(EGPD_Output, UEdGraphSchema_K2::PC_Struct, FInputActionValue::StaticStruct(), ActionPinName);
+	
 	Super::AllocateDefaultPins();
 }
 
@@ -206,6 +210,11 @@ UEdGraphPin* UK2Node_InputDebugKey::GetReleasedPin() const
 	return FindPin(TEXT("Released"));
 }
 
+UEdGraphPin* UK2Node_InputDebugKey::GetActionValuePin() const
+{
+	return FindPin(ActionPinName);
+}
+
 void UK2Node_InputDebugKey::ValidateNodeDuringCompilation(class FCompilerResultsLog& MessageLog) const
 {
 	Super::ValidateNodeDuringCompilation(MessageLog);
@@ -226,7 +235,8 @@ void UK2Node_InputDebugKey::ExpandNode(FKismetCompilerContext& CompilerContext,
 
 	UEdGraphPin* InputDebugKeyPressedPin = GetPressedPin();
 	UEdGraphPin* InputDebugKeyReleasedPin = GetReleasedPin();
-
+	UEdGraphPin* ActionValuePin = GetActionValuePin();
+	
 	struct EventPinData
 	{
 		EventPinData(UEdGraphPin* InPin,TEnumAsByte<EInputEvent> InEvent ){	Pin=InPin;EventType=InEvent;};
@@ -282,6 +292,13 @@ void UK2Node_InputDebugKey::ExpandNode(FKismetCompilerContext& CompilerContext,
 		KeyVar->VariableType.PinSubCategoryObject = KeyStruct;
 		KeyVar->AllocateDefaultPins();
 
+		// Create a temporary variable to copy Action Value in to
+		static UScriptStruct* ActionValueStruct = FInputActionValue::StaticStruct();
+		UK2Node_TemporaryVariable* ActionValueVar = CompilerContext.SpawnIntermediateNode<UK2Node_TemporaryVariable>(this, SourceGraph);
+		ActionValueVar->VariableType.PinCategory = UEdGraphSchema_K2::PC_Struct;
+		ActionValueVar->VariableType.PinSubCategoryObject = ActionValueStruct;
+		ActionValueVar->AllocateDefaultPins();
+
 		for (auto PinIt = ActivePins.CreateIterator(); PinIt; ++PinIt)
 		{
 			UEdGraphPin *EachPin = (*PinIt).Pin;
@@ -293,14 +310,23 @@ void UK2Node_InputDebugKey::ExpandNode(FKismetCompilerContext& CompilerContext,
 			KeyInitialize->AllocateDefaultPins();
 			Schema->TryCreateConnection(KeyVar->GetVariablePin(), KeyInitialize->GetVariablePin());
 			Schema->TryCreateConnection(KeyInitialize->GetValuePin(), InputDebugKeyEvent->FindPinChecked(TEXT("Key")));
+
+			// Create Assignment nodes to assign the action value
+			UK2Node_AssignmentStatement* ActionValueInitialize = CompilerContext.SpawnIntermediateNode<UK2Node_AssignmentStatement>(this, SourceGraph);
+			ActionValueInitialize->AllocateDefaultPins();
+			Schema->TryCreateConnection(ActionValueVar->GetVariablePin(), ActionValueInitialize->GetVariablePin());
+			Schema->TryCreateConnection(ActionValueInitialize->GetValuePin(), InputDebugKeyEvent->FindPinChecked(ActionPinName));
+			
 			// Connect the events to the assign key nodes
 			Schema->TryCreateConnection(Schema->FindExecutionPin(*InputDebugKeyEvent, EGPD_Output), KeyInitialize->GetExecPin());
-
+			Schema->TryCreateConnection(KeyInitialize->GetThenPin(), ActionValueInitialize->GetExecPin());
+			
 			// Move the original event connections to the then pin of the key assign
-			CompilerContext.MovePinLinksToIntermediate(*EachPin, *KeyInitialize->GetThenPin());
+			CompilerContext.MovePinLinksToIntermediate(*EachPin, *ActionValueInitialize->GetThenPin());
 
 			// Move the original event variable connections to the intermediate nodes
 			CompilerContext.MovePinLinksToIntermediate(*FindPin(TEXT("Key")), *KeyVar->GetVariablePin());
+			CompilerContext.MovePinLinksToIntermediate(*ActionValuePin, *ActionValueVar->GetVariablePin());
 		}
 	}
 	else if( ActivePins.Num() == 1 )
@@ -313,6 +339,7 @@ void UK2Node_InputDebugKey::ExpandNode(FKismetCompilerContext& CompilerContext,
 
 			CompilerContext.MovePinLinksToIntermediate(*InputDebugKeyPin, *Schema->FindExecutionPin(*InputDebugKeyEvent, EGPD_Output));
 			CompilerContext.MovePinLinksToIntermediate(*FindPin(TEXT("Key")), *InputDebugKeyEvent->FindPin(TEXT("Key")));
+			CompilerContext.MovePinLinksToIntermediate(*ActionValuePin, *InputDebugKeyEvent->FindPin(ActionPinName));
 		}
 	}
 }
diff --git a/Engine/Plugins/Experimental/EnhancedInput/Source/InputBlueprintNodes/Private/K2Node_InputDebugKey.h b/Engine/Plugins/Experimental/EnhancedInput/Source/InputBlueprintNodes/Private/K2Node_InputDebugKey.h
--- a/Engine/Plugins/Experimental/EnhancedInput/Source/InputBlueprintNodes/Private/K2Node_InputDebugKey.h
+++ b/Engine/Plugins/Experimental/EnhancedInput/Source/InputBlueprintNodes/Private/K2Node_InputDebugKey.h
@@ -85,6 +85,9 @@ class INPUTBLUEPRINTNODES_API UK2Node_InputDebugKey : public UK2Node, public IK2
 	/** Get the 'released' input pin */
 	UEdGraphPin* GetReleasedPin() const;
 
+	/** Get the 'Action Value' input pin */
+	UEdGraphPin* GetActionValuePin() const;
+
 private:
 	/** Constructing FText strings can be costly, so we cache the node's title/tooltip */
 	FNodeTextCache CachedTooltip;
