Expose bForceRebuildImmediately to AddMappingContext and RemoveMappingContext, allowing users to more easily force a rebuild of the control mappings within the same frame.

diff --git a/Engine/Plugins/Experimental/EnhancedInput/Source/EnhancedInput/Private/EnhancedInputLibrary.cpp b/Engine/Plugins/Experimental/EnhancedInput/Source/EnhancedInput/Private/EnhancedInputLibrary.cpp
--- a/Engine/Plugins/Experimental/EnhancedInput/Source/EnhancedInput/Private/EnhancedInputLibrary.cpp
+++ b/Engine/Plugins/Experimental/EnhancedInput/Source/EnhancedInput/Private/EnhancedInputLibrary.cpp
@@ -32,7 +32,9 @@ void UEnhancedInputLibrary::RequestRebuildControlMappingsUsingContext(const UInp
 			check(Subsystem);
 			if (Subsystem && Subsystem->HasMappingContext(Context))
 			{
-				Subsystem->RequestRebuildControlMappings(bForceImmediately);
+				FModifyContextOptions Options {};
+				Options.bForceImmediately = bForceImmediately;
+				Subsystem->RequestRebuildInputControlMappings(Options);
 			}
 		});
 }
diff --git a/Engine/Plugins/Experimental/EnhancedInput/Source/EnhancedInput/Private/EnhancedInputSubsystemInterface.cpp b/Engine/Plugins/Experimental/EnhancedInput/Source/EnhancedInput/Private/EnhancedInputSubsystemInterface.cpp
--- a/Engine/Plugins/Experimental/EnhancedInput/Source/EnhancedInput/Private/EnhancedInputSubsystemInterface.cpp
+++ b/Engine/Plugins/Experimental/EnhancedInput/Source/EnhancedInput/Private/EnhancedInputSubsystemInterface.cpp
@@ -40,11 +40,19 @@ void IEnhancedInputSubsystemInterface::ClearAllMappings()
 	if (UEnhancedPlayerInput* PlayerInput = GetPlayerInput())
 	{
 		PlayerInput->AppliedInputContexts.Empty();
-		RequestRebuildControlMappings(false);
+		RequestRebuildInputControlMappings();
 	}
 }
 
 void IEnhancedInputSubsystemInterface::AddMappingContext(const UInputMappingContext* MappingContext, int32 Priority, const bool bIgnoreAllPressedKeysUntilRelease /* = true */)
+{
+	FModifyContextOptions Options {};
+	Options.bIgnoreAllPressedKeysUntilRelease = bIgnoreAllPressedKeysUntilRelease;
+	
+	AddInputMappingContext(MappingContext, Priority, Options);
+}
+
+void IEnhancedInputSubsystemInterface::AddInputMappingContext(const UInputMappingContext* MappingContext, int32 Priority, const FModifyContextOptions& Options)
 {
 	// Layer mappings on top of existing mappings
 	if (MappingContext)
@@ -52,35 +60,51 @@ void IEnhancedInputSubsystemInterface::AddMappingContext(const UInputMappingCont
 		if (UEnhancedPlayerInput* PlayerInput = GetPlayerInput())
 		{
 			PlayerInput->AppliedInputContexts.Add(MappingContext, Priority);
-			RequestRebuildControlMappings(false, bIgnoreAllPressedKeysUntilRelease);
+			RequestRebuildInputControlMappings(Options);
 		}
 	}
 }
 
 void IEnhancedInputSubsystemInterface::RemoveMappingContext(const UInputMappingContext* MappingContext, const bool bIgnoreAllPressedKeysUntilRelease /* = true */)
+{
+	FModifyContextOptions Options {};
+	Options.bIgnoreAllPressedKeysUntilRelease = bIgnoreAllPressedKeysUntilRelease;
+	
+	RemoveInputMappingContext(MappingContext, Options);
+}
+
+void IEnhancedInputSubsystemInterface::RemoveInputMappingContext(const UInputMappingContext* MappingContext, const FModifyContextOptions& Options)
 {
 	if (MappingContext)
 	{
 		if (UEnhancedPlayerInput* PlayerInput = GetPlayerInput())
 		{
 			PlayerInput->AppliedInputContexts.Remove(MappingContext);
-			RequestRebuildControlMappings(false, bIgnoreAllPressedKeysUntilRelease);
+			RequestRebuildInputControlMappings(Options);
 		}
 	}
 }
 
 void IEnhancedInputSubsystemInterface::RequestRebuildControlMappings(bool bForceImmediately, const bool bIgnoreAllPressedKeysUntilRelease /* = true */)
+{
+	FModifyContextOptions Options;
+	Options.bForceImmediately = bForceImmediately;
+	Options.bIgnoreAllPressedKeysUntilRelease = bIgnoreAllPressedKeysUntilRelease;
+	
+	RequestRebuildInputControlMappings(Options);
+}
+
+void IEnhancedInputSubsystemInterface::RequestRebuildInputControlMappings(const FModifyContextOptions& Options)
 {
 	bMappingRebuildPending = true;
-	bIgnoreAllPressedKeysUntilReleaseOnRebuild &= bIgnoreAllPressedKeysUntilRelease;
+	bIgnoreAllPressedKeysUntilReleaseOnRebuild &= Options.bIgnoreAllPressedKeysUntilRelease;
 
-	if (bForceImmediately)
+	if (Options.bForceImmediately)
 	{
 		RebuildControlMappings();
 	}
 }
 
-
 EMappingQueryResult IEnhancedInputSubsystemInterface::QueryMapKeyInActiveContextSet(const UInputMappingContext* InputContext, const UInputAction* Action, FKey Key, TArray<FMappingQueryIssue>& OutIssues, EMappingQueryIssue BlockingIssues/* = DefaultMappingIssues::StandardFatal*/)
 {
 	UEnhancedPlayerInput* PlayerInput = GetPlayerInput();
diff --git a/Engine/Plugins/Experimental/EnhancedInput/Source/EnhancedInput/Public/EnhancedInputSubsystemInterface.h b/Engine/Plugins/Experimental/EnhancedInput/Source/EnhancedInput/Public/EnhancedInputSubsystemInterface.h
--- a/Engine/Plugins/Experimental/EnhancedInput/Source/EnhancedInput/Public/EnhancedInputSubsystemInterface.h
+++ b/Engine/Plugins/Experimental/EnhancedInput/Source/EnhancedInput/Public/EnhancedInputSubsystemInterface.h
@@ -23,6 +23,26 @@ class UEnhancedInputSubsystemInterface : public UInterface
 	GENERATED_BODY()
 };
 
+/** Passed in as params for Adding/Remove input contexts */
+USTRUCT(BlueprintType)
+struct FModifyContextOptions
+{
+	GENERATED_BODY()
+	
+	FModifyContextOptions()
+		: bIgnoreAllPressedKeysUntilRelease(true)
+		, bForceImmediately(false)
+	{}
+
+	// If true than any keys that are pressed during the rebuild of control mappings will be ignored until they are released.
+	UPROPERTY(EditAnywhere, BlueprintReadWrite, Category = "Input")
+	uint8 bIgnoreAllPressedKeysUntilRelease : 1;
+
+	// The mapping changes will be applied synchronously, rather than at the end of the frame, making them available to the input system on the same frame.
+	UPROPERTY(EditAnywhere, BlueprintReadWrite, Category = "Input")
+	uint8 bForceImmediately : 1;
+};
+
 // Includes native functionality shared between all subsystems
 class ENHANCEDINPUT_API IEnhancedInputSubsystemInterface
 {
@@ -54,28 +74,61 @@ public:
 
 	/**
 	 * Add a control mapping context.
-	 * @param MappingContext	A set of key to action mappings to apply to this player
-	 * @param Priority			Higher priority mappings will be applied first and, if they consume input, will block lower priority mappings.
+	 * 
+	 * @param MappingContext		A set of key to action mappings to apply to this player
+	 * @param Priority				Higher priority mappings will be applied first and, if they consume input, will block lower priority mappings.
+	 * @param bIgnoreAllPressedKeysUntilRelease	If true than any keys that are pressed during the rebuild of control mappings will be ignored until they are released.
 	 */
-	UFUNCTION(BlueprintCallable, BlueprintCosmetic, Category = "Input", meta=(AdvancedDisplay = "bIgnoreAllPressedKeysUntilRelease"))
+	UE_DEPRECATED(5.0, "This version of AddMappingContext has been deprecated, please use the version that takes in a FModifyContextOptions instead.")
+	UFUNCTION(BlueprintCallable, BlueprintCosmetic, Category = "Input", meta=(AdvancedDisplay = "bIgnoreAllPressedKeysUntilRelease", DisplayName= "AddMappingContext_DEPRECATED"))
 	virtual void AddMappingContext(const UInputMappingContext* MappingContext, int32 Priority, const bool bIgnoreAllPressedKeysUntilRelease = true);
-
+	
+	/**
+	 * Add a control mapping context.
+	 * @param MappingContext		A set of key to action mappings to apply to this player
+	 * @param Priority				Higher priority mappings will be applied first and, if they consume input, will block lower priority mappings.
+	 * @param Options				Options to consider when adding this mapping context.
+	 */
+	UFUNCTION(BlueprintCallable, BlueprintCosmetic, Category = "Input", meta=(DisplayName="Add Mapping Context", ScriptName="AddMappingContext", AutoCreateRefParm = "Options"))
+	virtual void AddInputMappingContext(const UInputMappingContext* MappingContext, int32 Priority, const FModifyContextOptions& Options = FModifyContextOptions());
+	
 	/**
 	 * Remove a specific control context. 
 	 * This is safe to call even if the context is not applied.
 	 * @param MappingContext		Context to remove from the player
+	 * @param bIgnoreAllPressedKeysUntilRelease	If true than any keys that are pressed during the rebuild of control mappings will be ignored until they are released.
 	 */
-	UFUNCTION(BlueprintCallable, BlueprintCosmetic, Category = "Input", meta=(AdvancedDisplay = "bIgnoreAllPressedKeysUntilRelease"))
+	UE_DEPRECATED(5.0, "This version of RemoveMappingContext has been deprecated, please use the version that takes in a FModifyContextOptions instead.")
+	UFUNCTION(BlueprintCallable, BlueprintCosmetic, Category = "Input", meta=(AdvancedDisplay = "bIgnoreAllPressedKeysUntilRelease", DisplayName= "RemoveMappingContext_DEPRECATED"))
 	virtual void RemoveMappingContext(const UInputMappingContext* MappingContext, const bool bIgnoreAllPressedKeysUntilRelease = true);
 
+	/**
+	* Remove a specific control context. 
+	* This is safe to call even if the context is not applied.
+	* @param MappingContext		Context to remove from the player
+	* @param Options			Options to consider when removing this input mapping context
+	*/
+	UFUNCTION(BlueprintCallable, BlueprintCosmetic, Category = "Input", meta=(DisplayName = "Remove Mapping Context", ScriptName = "RemoveMappingContext", AutoCreateRefParm = "Options"))
+	virtual void RemoveInputMappingContext(const UInputMappingContext* MappingContext, const FModifyContextOptions& Options = FModifyContextOptions());
+	
 	/**
 	 * Flag player for reapplication of all mapping contexts at the end of this frame.
 	 * This is called automatically when adding or removing mappings contexts.
 	 * @param bForceImmediately		THe mapping changes will be applied synchronously, rather than at the end of the frame, making them available to the input system on the same frame.
 	 */
-	UFUNCTION(BlueprintCallable, BlueprintCosmetic, Category = "Input", meta=(AdvancedDisplay = "bIgnoreAllPressedKeysUntilRelease"))
+	UE_DEPRECATED(5.0, "This version of RequestRebuildControlMappings has been deprecated, please use the version that takes in a FModifyContextOptions instead.")
+	UFUNCTION(BlueprintCallable, BlueprintCosmetic, Category = "Input", meta=(AdvancedDisplay = "bIgnoreAllPressedKeysUntilRelease", DisplayName="RequestRebuildControlMappings_DEPRECATED"))
 	virtual void RequestRebuildControlMappings(bool bForceImmediately = false, const bool bIgnoreAllPressedKeysUntilRelease = true);
 
+	/**
+	* Flag player for reapplication of all mapping contexts at the end of this frame.
+	* This is called automatically when adding or removing mappings contexts.
+	*
+	* @param Options		Options to consider when removing this input mapping context
+	*/
+	UFUNCTION(BlueprintCallable, BlueprintCosmetic, Category = "Input", meta=(DisplayName = "RequestRebuildControlMappings", ScriptName="RequestRebuildControlMappings", AutoCreateRefParm = "Options"))
+	virtual void RequestRebuildInputControlMappings(const FModifyContextOptions& Options = FModifyContextOptions());
+
 	/**
 	 * Check if a key mapping is safe to add to a given mapping context within the set of active contexts currently applied to the player controller.
 	 * @param InputContext		Mapping context to which the action/key mapping is intended to be added
diff --git a/Engine/Plugins/Experimental/EnhancedInput/Source/InputEditor/Private/Tests/InputTestFramework.cpp b/Engine/Plugins/Experimental/EnhancedInput/Source/InputEditor/Private/Tests/InputTestFramework.cpp
--- a/Engine/Plugins/Experimental/EnhancedInput/Source/InputEditor/Private/Tests/InputTestFramework.cpp
+++ b/Engine/Plugins/Experimental/EnhancedInput/Source/InputEditor/Private/Tests/InputTestFramework.cpp
@@ -46,7 +46,7 @@ UControllablePlayer& AControllablePlayer(UWorld* World)
 UInputMappingContext* AnInputContextIsAppliedToAPlayer(UControllablePlayer& PlayerData, FName ContextName, int32 WithPriority)
 {
 	UInputMappingContext* Context = PlayerData.InputContext.Emplace(ContextName, NewObject<UInputMappingContext>(PlayerData.Player, ContextName));
-	PlayerData.Subsystem->AddMappingContext(Context, WithPriority);
+	PlayerData.Subsystem->AddInputMappingContext(Context, WithPriority);
 	return Context;
 }
 
@@ -60,7 +60,9 @@ UInputAction* AnInputAction(UControllablePlayer& PlayerData, FName ActionName, E
 void ControlMappingsAreRebuilt(UControllablePlayer& PlayerData)
 {
 	FInputTestHelper::ResetActionInstanceData(PlayerData);
-	PlayerData.Subsystem->RequestRebuildControlMappings(true);
+	FModifyContextOptions Options;
+	Options.bForceImmediately = true;
+	PlayerData.Subsystem->RequestRebuildInputControlMappings(Options);
 }
 
 FEnhancedActionKeyMapping& AnActionIsMappedToAKey(UControllablePlayer& PlayerData, FName ContextName, FName ActionName, FKey Key)
