Fix ScriptName collision for Enhanced Input

#rnx

diff --git a/Engine/Plugins/Experimental/EnhancedInput/Source/EnhancedInput/Public/EnhancedInputSubsystemInterface.h b/Engine/Plugins/Experimental/EnhancedInput/Source/EnhancedInput/Public/EnhancedInputSubsystemInterface.h
--- a/Engine/Plugins/Experimental/EnhancedInput/Source/EnhancedInput/Public/EnhancedInputSubsystemInterface.h
+++ b/Engine/Plugins/Experimental/EnhancedInput/Source/EnhancedInput/Public/EnhancedInputSubsystemInterface.h
@@ -80,7 +80,7 @@ public:
 	 * @param bIgnoreAllPressedKeysUntilRelease	If true than any keys that are pressed during the rebuild of control mappings will be ignored until they are released.
 	 */
 	UE_DEPRECATED(5.0, "This version of AddMappingContext has been deprecated, please use the version that takes in a FModifyContextOptions instead.")
-	UFUNCTION(BlueprintCallable, BlueprintCosmetic, Category = "Input", meta=(AdvancedDisplay = "bIgnoreAllPressedKeysUntilRelease", DisplayName= "AddMappingContext_DEPRECATED"))
+	UFUNCTION(BlueprintCallable, BlueprintCosmetic, Category = "Input", meta=(AdvancedDisplay = "bIgnoreAllPressedKeysUntilRelease", ScriptName = "AddMappingContext_DEPRECATED", DisplayName= "AddMappingContext_DEPRECATED"))
 	virtual void AddMappingContext(const UInputMappingContext* MappingContext, int32 Priority, const bool bIgnoreAllPressedKeysUntilRelease = true);
 	
 	/**
@@ -99,7 +99,7 @@ public:
 	 * @param bIgnoreAllPressedKeysUntilRelease	If true than any keys that are pressed during the rebuild of control mappings will be ignored until they are released.
 	 */
 	UE_DEPRECATED(5.0, "This version of RemoveMappingContext has been deprecated, please use the version that takes in a FModifyContextOptions instead.")
-	UFUNCTION(BlueprintCallable, BlueprintCosmetic, Category = "Input", meta=(AdvancedDisplay = "bIgnoreAllPressedKeysUntilRelease", DisplayName= "RemoveMappingContext_DEPRECATED"))
+	UFUNCTION(BlueprintCallable, BlueprintCosmetic, Category = "Input", meta=(AdvancedDisplay = "bIgnoreAllPressedKeysUntilRelease", ScriptName = "RemoveMappingContext_DEPRECATED", DisplayName = "RemoveMappingContext_DEPRECATED"))
 	virtual void RemoveMappingContext(const UInputMappingContext* MappingContext, const bool bIgnoreAllPressedKeysUntilRelease = true);
 
 	/**
@@ -117,7 +117,7 @@ public:
 	 * @param bForceImmediately		THe mapping changes will be applied synchronously, rather than at the end of the frame, making them available to the input system on the same frame.
 	 */
 	UE_DEPRECATED(5.0, "This version of RequestRebuildControlMappings has been deprecated, please use the version that takes in a FModifyContextOptions instead.")
-	UFUNCTION(BlueprintCallable, BlueprintCosmetic, Category = "Input", meta=(AdvancedDisplay = "bIgnoreAllPressedKeysUntilRelease", DisplayName="RequestRebuildControlMappings_DEPRECATED"))
+	UFUNCTION(BlueprintCallable, BlueprintCosmetic, Category = "Input", meta=(AdvancedDisplay = "bIgnoreAllPressedKeysUntilRelease", ScriptName = "RequestRebuildControlMappings_DEPRECATED", DisplayName="RequestRebuildControlMappings_DEPRECATED"))
 	virtual void RequestRebuildControlMappings(bool bForceImmediately = false, const bool bIgnoreAllPressedKeysUntilRelease = true);
 
 	/**
