Track the real time delta seconds of the world to use as a backup if Time Dilation on the player controller is set to 0. This ensures that timed triggers are still calculated correctly even with time dilation at 0.

diff --git a/Engine/Plugins/Experimental/EnhancedInput/Source/EnhancedInput/Private/EnhancedPlayerInput.cpp b/Engine/Plugins/Experimental/EnhancedInput/Source/EnhancedInput/Private/EnhancedPlayerInput.cpp
--- a/Engine/Plugins/Experimental/EnhancedInput/Source/EnhancedInput/Private/EnhancedPlayerInput.cpp
+++ b/Engine/Plugins/Experimental/EnhancedInput/Source/EnhancedInput/Private/EnhancedPlayerInput.cpp
@@ -5,7 +5,7 @@
 #include "GameFramework/PlayerController.h"
 #include "InputModifiers.h"
 #include "InputTriggers.h"
-
+#include "Engine/World.h"
 
 // NOTE: Enum order represents firing priority(lowest to highest) and is important as multiple keys bound to the same action may generate differing trigger event states.
 enum class ETriggerEventInternal : uint8
@@ -188,10 +188,18 @@ void UEnhancedPlayerInput::ProcessInputStack(const TArray<UInputComponent*>& Inp
 	// Process Action bindings
 	ActionsWithEventsThisTick.Reset();
 
+	// Calculate the current delta time as a fallback in case the Time Dilation is set to 0.0
+	// This will ensure that Timed Triggers are calculated in real time
+	const UWorld* World = GetWorld();
+	// If there is no world, then add 1/60 to the last frame time to try and recover at least some
+	// semblance of real time passing. 
+	float CurrentTime = World ? World->GetRealTimeSeconds() : LastFrameTime + (1.0f / 60.0f);
+	RealTimeDeltaSeconds = CurrentTime - LastFrameTime;
+	
 	// Use non-dilated delta time for processing
 	check(GetOuterAPlayerController());
 	const float Dilation = GetOuterAPlayerController()->GetActorTimeDilation();
-	const float NonDilatedDeltaTime = DeltaTime / Dilation;
+	const float NonDilatedDeltaTime = Dilation != 0.0f ? DeltaTime / Dilation : RealTimeDeltaSeconds;
 
 	// Handle input devices, applying modifiers and triggers
 	for (FEnhancedActionKeyMapping& Mapping : EnhancedActionMappings)
@@ -463,6 +471,8 @@ void UEnhancedPlayerInput::ProcessInputStack(const TArray<UInputComponent*>& Inp
 		ActionData.MappingTriggerState = ETriggerState::None;
 		ActionData.bMappingTriggerApplied = false;
 	}
+
+	LastFrameTime = CurrentTime;
 }
 
 FInputActionValue UEnhancedPlayerInput::GetActionValue(const UInputAction* ForAction) const
diff --git a/Engine/Plugins/Experimental/EnhancedInput/Source/EnhancedInput/Public/EnhancedPlayerInput.h b/Engine/Plugins/Experimental/EnhancedInput/Source/EnhancedInput/Public/EnhancedPlayerInput.h
--- a/Engine/Plugins/Experimental/EnhancedInput/Source/EnhancedInput/Public/EnhancedPlayerInput.h
+++ b/Engine/Plugins/Experimental/EnhancedInput/Source/EnhancedInput/Public/EnhancedPlayerInput.h
@@ -108,4 +108,10 @@ private:
 
 	/** Last frame's injected inputs */
 	TSet<const UInputAction*> LastInjectedActions;
+
+	/** The last time of the last frame that was processed in ProcessPlayerInput */
+	float LastFrameTime = 0.0f;
+
+	/** Delta seconds between frames calculated with UWorld::GetRealTimeSeconds */
+	float RealTimeDeltaSeconds = 0.0f;
 };
