Deprecate AddMappingContext in a better way making it more friendly to native C++ deprecations.

diff --git a/Engine/Plugins/Experimental/EnhancedInput/Source/EnhancedInput/Private/EnhancedInputLibrary.cpp b/Engine/Plugins/Experimental/EnhancedInput/Source/EnhancedInput/Private/EnhancedInputLibrary.cpp
--- a/Engine/Plugins/Experimental/EnhancedInput/Source/EnhancedInput/Private/EnhancedInputLibrary.cpp
+++ b/Engine/Plugins/Experimental/EnhancedInput/Source/EnhancedInput/Private/EnhancedInputLibrary.cpp
@@ -34,7 +34,7 @@ void UEnhancedInputLibrary::RequestRebuildControlMappingsUsingContext(const UInp
 			{
 				FModifyContextOptions Options {};
 				Options.bForceImmediately = bForceImmediately;
-				Subsystem->RequestRebuildInputControlMappings(Options);
+				Subsystem->RequestRebuildControlMappings(Options);
 			}
 		});
 }
diff --git a/Engine/Plugins/Experimental/EnhancedInput/Source/EnhancedInput/Private/EnhancedInputSubsystemInterface.cpp b/Engine/Plugins/Experimental/EnhancedInput/Source/EnhancedInput/Private/EnhancedInputSubsystemInterface.cpp
--- a/Engine/Plugins/Experimental/EnhancedInput/Source/EnhancedInput/Private/EnhancedInputSubsystemInterface.cpp
+++ b/Engine/Plugins/Experimental/EnhancedInput/Source/EnhancedInput/Private/EnhancedInputSubsystemInterface.cpp
@@ -40,19 +40,19 @@ void IEnhancedInputSubsystemInterface::ClearAllMappings()
 	if (UEnhancedPlayerInput* PlayerInput = GetPlayerInput())
 	{
 		PlayerInput->AppliedInputContexts.Empty();
-		RequestRebuildInputControlMappings();
+		RequestRebuildControlMappings();
 	}
 }
 
-void IEnhancedInputSubsystemInterface::AddMappingContext(const UInputMappingContext* MappingContext, int32 Priority, const bool bIgnoreAllPressedKeysUntilRelease /* = true */)
+void IEnhancedInputSubsystemInterface::AddMappingContext(const UInputMappingContext* MappingContext, int32 Priority, const bool bIgnoreAllPressedKeysUntilRelease)
 {
 	FModifyContextOptions Options {};
 	Options.bIgnoreAllPressedKeysUntilRelease = bIgnoreAllPressedKeysUntilRelease;
 	
-	AddInputMappingContext(MappingContext, Priority, Options);
+	AddMappingContext(MappingContext, Priority, Options);
 }
 
-void IEnhancedInputSubsystemInterface::AddInputMappingContext(const UInputMappingContext* MappingContext, int32 Priority, const FModifyContextOptions& Options)
+void IEnhancedInputSubsystemInterface::AddMappingContext(const UInputMappingContext* MappingContext, int32 Priority, const FModifyContextOptions& Options)
 {
 	// Layer mappings on top of existing mappings
 	if (MappingContext)
@@ -60,41 +60,41 @@ void IEnhancedInputSubsystemInterface::AddInputMappingContext(const UInputMappin
 		if (UEnhancedPlayerInput* PlayerInput = GetPlayerInput())
 		{
 			PlayerInput->AppliedInputContexts.Add(MappingContext, Priority);
-			RequestRebuildInputControlMappings(Options);
+			RequestRebuildControlMappings(Options);
 		}
 	}
 }
 
-void IEnhancedInputSubsystemInterface::RemoveMappingContext(const UInputMappingContext* MappingContext, const bool bIgnoreAllPressedKeysUntilRelease /* = true */)
+void IEnhancedInputSubsystemInterface::RemoveMappingContext(const UInputMappingContext* MappingContext, const bool bIgnoreAllPressedKeysUntilRelease)
 {
 	FModifyContextOptions Options {};
 	Options.bIgnoreAllPressedKeysUntilRelease = bIgnoreAllPressedKeysUntilRelease;
 	
-	RemoveInputMappingContext(MappingContext, Options);
+	RemoveMappingContext(MappingContext, Options);
 }
 
-void IEnhancedInputSubsystemInterface::RemoveInputMappingContext(const UInputMappingContext* MappingContext, const FModifyContextOptions& Options)
+void IEnhancedInputSubsystemInterface::RemoveMappingContext(const UInputMappingContext* MappingContext, const FModifyContextOptions& Options)
 {
 	if (MappingContext)
 	{
 		if (UEnhancedPlayerInput* PlayerInput = GetPlayerInput())
 		{
 			PlayerInput->AppliedInputContexts.Remove(MappingContext);
-			RequestRebuildInputControlMappings(Options);
+			RequestRebuildControlMappings(Options);
 		}
 	}
 }
 
-void IEnhancedInputSubsystemInterface::RequestRebuildControlMappings(bool bForceImmediately, const bool bIgnoreAllPressedKeysUntilRelease /* = true */)
+void IEnhancedInputSubsystemInterface::RequestRebuildControlMappings(bool bForceImmediately, const bool bIgnoreAllPressedKeysUntilRelease)
 {
 	FModifyContextOptions Options;
 	Options.bForceImmediately = bForceImmediately;
 	Options.bIgnoreAllPressedKeysUntilRelease = bIgnoreAllPressedKeysUntilRelease;
 	
-	RequestRebuildInputControlMappings(Options);
+	RequestRebuildControlMappings(Options);
 }
 
-void IEnhancedInputSubsystemInterface::RequestRebuildInputControlMappings(const FModifyContextOptions& Options)
+void IEnhancedInputSubsystemInterface::RequestRebuildControlMappings(const FModifyContextOptions& Options)
 {
 	bMappingRebuildPending = true;
 	bIgnoreAllPressedKeysUntilReleaseOnRebuild &= Options.bIgnoreAllPressedKeysUntilRelease;
diff --git a/Engine/Plugins/Experimental/EnhancedInput/Source/EnhancedInput/Public/EnhancedInputSubsystemInterface.h b/Engine/Plugins/Experimental/EnhancedInput/Source/EnhancedInput/Public/EnhancedInputSubsystemInterface.h
--- a/Engine/Plugins/Experimental/EnhancedInput/Source/EnhancedInput/Public/EnhancedInputSubsystemInterface.h
+++ b/Engine/Plugins/Experimental/EnhancedInput/Source/EnhancedInput/Public/EnhancedInputSubsystemInterface.h
@@ -80,8 +80,7 @@ public:
 	 * @param bIgnoreAllPressedKeysUntilRelease	If true than any keys that are pressed during the rebuild of control mappings will be ignored until they are released.
 	 */
 	UE_DEPRECATED(5.0, "This version of AddMappingContext has been deprecated, please use the version that takes in a FModifyContextOptions instead.")
-	UFUNCTION(BlueprintCallable, BlueprintCosmetic, Category = "Input", meta=(AdvancedDisplay = "bIgnoreAllPressedKeysUntilRelease", ScriptName = "AddMappingContext_DEPRECATED", DisplayName= "AddMappingContext_DEPRECATED"))
-	virtual void AddMappingContext(const UInputMappingContext* MappingContext, int32 Priority, const bool bIgnoreAllPressedKeysUntilRelease = true);
+	virtual void AddMappingContext(const UInputMappingContext* MappingContext, int32 Priority, const bool bIgnoreAllPressedKeysUntilRelease);
 	
 	/**
 	 * Add a control mapping context.
@@ -89,8 +88,8 @@ public:
 	 * @param Priority				Higher priority mappings will be applied first and, if they consume input, will block lower priority mappings.
 	 * @param Options				Options to consider when adding this mapping context.
 	 */
-	UFUNCTION(BlueprintCallable, BlueprintCosmetic, Category = "Input", meta=(DisplayName="Add Mapping Context", ScriptName="AddMappingContext", AutoCreateRefTerm = "Options"))
-	virtual void AddInputMappingContext(const UInputMappingContext* MappingContext, int32 Priority, const FModifyContextOptions& Options = FModifyContextOptions());
+	UFUNCTION(BlueprintCallable, BlueprintCosmetic, Category = "Input", meta=(AutoCreateRefTerm = "Options"))
+	virtual void AddMappingContext(const UInputMappingContext* MappingContext, int32 Priority, const FModifyContextOptions& Options = FModifyContextOptions());
 	
 	/**
 	 * Remove a specific control context. 
@@ -99,8 +98,7 @@ public:
 	 * @param bIgnoreAllPressedKeysUntilRelease	If true than any keys that are pressed during the rebuild of control mappings will be ignored until they are released.
 	 */
 	UE_DEPRECATED(5.0, "This version of RemoveMappingContext has been deprecated, please use the version that takes in a FModifyContextOptions instead.")
-	UFUNCTION(BlueprintCallable, BlueprintCosmetic, Category = "Input", meta=(AdvancedDisplay = "bIgnoreAllPressedKeysUntilRelease", ScriptName = "RemoveMappingContext_DEPRECATED", DisplayName = "RemoveMappingContext_DEPRECATED"))
-	virtual void RemoveMappingContext(const UInputMappingContext* MappingContext, const bool bIgnoreAllPressedKeysUntilRelease = true);
+	virtual void RemoveMappingContext(const UInputMappingContext* MappingContext, const bool bIgnoreAllPressedKeysUntilRelease);
 
 	/**
 	* Remove a specific control context. 
@@ -108,8 +106,8 @@ public:
 	* @param MappingContext		Context to remove from the player
 	* @param Options			Options to consider when removing this input mapping context
 	*/
-	UFUNCTION(BlueprintCallable, BlueprintCosmetic, Category = "Input", meta=(DisplayName = "Remove Mapping Context", ScriptName = "RemoveMappingContext", AutoCreateRefTerm = "Options"))
-	virtual void RemoveInputMappingContext(const UInputMappingContext* MappingContext, const FModifyContextOptions& Options = FModifyContextOptions());
+	UFUNCTION(BlueprintCallable, BlueprintCosmetic, Category = "Input", meta=(AutoCreateRefTerm = "Options"))
+	virtual void RemoveMappingContext(const UInputMappingContext* MappingContext, const FModifyContextOptions& Options = FModifyContextOptions());
 	
 	/**
 	 * Flag player for reapplication of all mapping contexts at the end of this frame.
@@ -117,8 +115,7 @@ public:
 	 * @param bForceImmediately		THe mapping changes will be applied synchronously, rather than at the end of the frame, making them available to the input system on the same frame.
 	 */
 	UE_DEPRECATED(5.0, "This version of RequestRebuildControlMappings has been deprecated, please use the version that takes in a FModifyContextOptions instead.")
-	UFUNCTION(BlueprintCallable, BlueprintCosmetic, Category = "Input", meta=(AdvancedDisplay = "bIgnoreAllPressedKeysUntilRelease", ScriptName = "RequestRebuildControlMappings_DEPRECATED", DisplayName="RequestRebuildControlMappings_DEPRECATED"))
-	virtual void RequestRebuildControlMappings(bool bForceImmediately = false, const bool bIgnoreAllPressedKeysUntilRelease = true);
+	virtual void RequestRebuildControlMappings(bool bForceImmediately, const bool bIgnoreAllPressedKeysUntilRelease = true);
 
 	/**
 	* Flag player for reapplication of all mapping contexts at the end of this frame.
@@ -126,8 +123,8 @@ public:
 	*
 	* @param Options		Options to consider when removing this input mapping context
 	*/
-	UFUNCTION(BlueprintCallable, BlueprintCosmetic, Category = "Input", meta=(DisplayName = "RequestRebuildControlMappings", ScriptName="RequestRebuildControlMappings", AutoCreateRefTerm = "Options"))
-	virtual void RequestRebuildInputControlMappings(const FModifyContextOptions& Options = FModifyContextOptions());
+	UFUNCTION(BlueprintCallable, BlueprintCosmetic, Category = "Input", meta=(AutoCreateRefTerm = "Options"))
+	virtual void RequestRebuildControlMappings(const FModifyContextOptions& Options = FModifyContextOptions());
 
 	/**
 	 * Check if a key mapping is safe to add to a given mapping context within the set of active contexts currently applied to the player controller.
diff --git a/Engine/Plugins/Experimental/EnhancedInput/Source/InputEditor/Private/Tests/InputTestFramework.cpp b/Engine/Plugins/Experimental/EnhancedInput/Source/InputEditor/Private/Tests/InputTestFramework.cpp
--- a/Engine/Plugins/Experimental/EnhancedInput/Source/InputEditor/Private/Tests/InputTestFramework.cpp
+++ b/Engine/Plugins/Experimental/EnhancedInput/Source/InputEditor/Private/Tests/InputTestFramework.cpp
@@ -46,7 +46,7 @@ UControllablePlayer& AControllablePlayer(UWorld* World)
 UInputMappingContext* AnInputContextIsAppliedToAPlayer(UControllablePlayer& PlayerData, FName ContextName, int32 WithPriority)
 {
 	UInputMappingContext* Context = PlayerData.InputContext.Emplace(ContextName, NewObject<UInputMappingContext>(PlayerData.Player, ContextName));
-	PlayerData.Subsystem->AddInputMappingContext(Context, WithPriority);
+	PlayerData.Subsystem->AddMappingContext(Context, WithPriority);
 	return Context;
 }
 
@@ -62,7 +62,7 @@ void ControlMappingsAreRebuilt(UControllablePlayer& PlayerData)
 	FInputTestHelper::ResetActionInstanceData(PlayerData);
 	FModifyContextOptions Options;
 	Options.bForceImmediately = true;
-	PlayerData.Subsystem->RequestRebuildInputControlMappings(Options);
+	PlayerData.Subsystem->RequestRebuildControlMappings(Options);
 }
 
 FEnhancedActionKeyMapping& AnActionIsMappedToAKey(UControllablePlayer& PlayerData, FName ContextName, FName ActionName, FKey Key)
