Fix a crash when creating a BP that has a null input action UProperty

#rnx

diff --git a/Engine/Plugins/Experimental/EnhancedInput/Source/InputBlueprintNodes/Private/K2Node_EnhancedInputAction.cpp b/Engine/Plugins/Experimental/EnhancedInput/Source/InputBlueprintNodes/Private/K2Node_EnhancedInputAction.cpp
--- a/Engine/Plugins/Experimental/EnhancedInput/Source/InputBlueprintNodes/Private/K2Node_EnhancedInputAction.cpp
+++ b/Engine/Plugins/Experimental/EnhancedInput/Source/InputBlueprintNodes/Private/K2Node_EnhancedInputAction.cpp
@@ -80,10 +80,13 @@ void UK2Node_EnhancedInputAction::AllocateDefaultPins()
 	CreatePin(EGPD_Output, UEdGraphSchema_K2::PC_Float, ElapsedSecondsPinName)->bAdvancedView = true;
 	CreatePin(EGPD_Output, UEdGraphSchema_K2::PC_Float, TriggeredSecondsPinName)->bAdvancedView = true;
 
-	UEdGraphPin* ActionPin = CreatePin(EGPD_Output, UEdGraphSchema_K2::PC_Object, InputAction->GetClass(), InputActionPinName);
-	ActionPin->DefaultObject = const_cast<UObject*>(Cast<UObject>(InputAction));
-	ActionPin->DefaultValue = InputAction->GetName();
-	Schema->ConstructBasicPinTooltip(*ActionPin, LOCTEXT("InputActionPinDescription", "The input action that caused this event to fire"), ActionPin->PinToolTip);
+	if(InputAction)
+	{
+		UEdGraphPin* ActionPin = CreatePin(EGPD_Output, UEdGraphSchema_K2::PC_Object, InputAction->GetClass(), InputActionPinName);
+		ActionPin->DefaultObject = const_cast<UObject*>(Cast<UObject>(InputAction));
+		ActionPin->DefaultValue = InputAction->GetName();
+		Schema->ConstructBasicPinTooltip(*ActionPin, LOCTEXT("InputActionPinDescription", "The input action that caused this event to fire"), ActionPin->PinToolTip);	
+	}
 	
 	Super::AllocateDefaultPins();
 }
