Refactor the evaluation of triggers out to a FTriggerStateTracker struct.

All unit tests passing!

diff --git a/Engine/Plugins/Experimental/EnhancedInput/Source/EnhancedInput/Private/EnhancedInputSubsystemInterface.cpp b/Engine/Plugins/Experimental/EnhancedInput/Source/EnhancedInput/Private/EnhancedInputSubsystemInterface.cpp
--- a/Engine/Plugins/Experimental/EnhancedInput/Source/EnhancedInput/Private/EnhancedInputSubsystemInterface.cpp
+++ b/Engine/Plugins/Experimental/EnhancedInput/Source/EnhancedInput/Private/EnhancedInputSubsystemInterface.cpp
@@ -94,11 +94,12 @@ void IEnhancedInputSubsystemInterface::RequestRebuildControlMappings(bool bForce
 	RequestRebuildControlMappings(Options);
 }
 
-void IEnhancedInputSubsystemInterface::RequestRebuildControlMappings(const FModifyContextOptions& Options)
+void IEnhancedInputSubsystemInterface::RequestRebuildControlMappings(const FModifyContextOptions& Options, EInputMappingRebuildType MappingRebuildType)
 {
 	bMappingRebuildPending = true;
 	bIgnoreAllPressedKeysUntilReleaseOnRebuild &= Options.bIgnoreAllPressedKeysUntilRelease;
-
+	MappingRebuildPending = MappingRebuildType;
+	
 	if (Options.bForceImmediately)
 	{
 		RebuildControlMappings();
@@ -468,7 +469,7 @@ TArray<FEnhancedActionKeyMapping> ReorderMappings(const TArray<FEnhancedActionKe
 
 void IEnhancedInputSubsystemInterface::RebuildControlMappings()
 {
-	if(!bMappingRebuildPending)
+	if(MappingRebuildPending == EInputMappingRebuildType::None)
 	{
 		return;
 	}
@@ -476,7 +477,7 @@ void IEnhancedInputSubsystemInterface::RebuildControlMappings()
 	UEnhancedPlayerInput* PlayerInput = GetPlayerInput();
 	if (!PlayerInput)
 	{
-		// TODO: Prefer to reset bMappingRebuildPending here?
+		// TODO: Prefer to reset MappingRebuildPending here?
 		return;
 	}
 
@@ -523,9 +524,6 @@ void IEnhancedInputSubsystemInterface::RebuildControlMappings()
 
 				ApplyAxisPropertyModifiers(PlayerInput, NewMapping);
 
-				// Perform a modifier calculation pass on the default data to initialize values correctly.
-				PlayerInput->InitializeMappingActionModifiers(NewMapping);
-
 				// Re-instance triggers
 				DeepCopyPtrArray<UInputTrigger>(Mapping.Triggers, NewMapping.Triggers);
 
@@ -558,51 +556,67 @@ void IEnhancedInputSubsystemInterface::RebuildControlMappings()
 
 	PlayerInput->ForceRebuildingKeyMaps(false);
 
-	// Remove action instance data for actions that are not referenced in the new action mappings
-	TSet<const UInputAction*> RemovedActions;
-	for (TPair<const UInputAction*, FInputActionInstance>& ActionInstance : PlayerInput->ActionInstanceData)
+	// Clean out invalidated actions
+	if (MappingRebuildPending == EInputMappingRebuildType::RebuildWithFlush)
 	{
-		RemovedActions.Add(ActionInstance.Key);
+		PlayerInput->ActionInstanceData.Empty();
 	}
-
-	// Return true if the given FKey was in the old Player Input mappings
-	auto WasInOldMapping = [&OldMappings](const FKey& InKey) -> bool
+	else
 	{
-		return OldMappings.ContainsByPredicate(
-			[&InKey](const FEnhancedActionKeyMapping& OldMapping){ return OldMapping.Key == InKey; }
-			);
-	};
+		
+		// Remove action instance data for actions that are not referenced in the new action mappings
+		TSet<const UInputAction*> RemovedActions;
+		for (TPair<const UInputAction*, FInputActionInstance>& ActionInstance : PlayerInput->ActionInstanceData)
+		{
+			RemovedActions.Add(ActionInstance.Key);
+		}
+
+		// Return true if the given FKey was in the old Player Input mappings
+		auto WasInOldMapping = [&OldMappings](const FKey& InKey) -> bool
+		{
+			return OldMappings.ContainsByPredicate(
+				[&InKey](const FEnhancedActionKeyMapping& OldMapping){ return OldMapping.Key == InKey; }
+				);
+		};
 	
-	for (FEnhancedActionKeyMapping& Mapping : PlayerInput->EnhancedActionMappings)
-	{
-		RemovedActions.Remove(Mapping.Action);
+		for (FEnhancedActionKeyMapping& Mapping : PlayerInput->EnhancedActionMappings)
+		{
+			RemovedActions.Remove(Mapping.Action);
+
+			// Was this key pressed last frame? If so, then we need to mark it to be ignored by the PlayerInput
+			// until it is released to avoid re-processing a triggered event when it. This is only a problem if
+			// the key was in the old mapping and the new one
+			if(bIgnoreAllPressedKeysUntilReleaseOnRebuild && Mapping.Key.IsDigital() && WasInOldMapping(Mapping.Key))
+			{				
+				const FKeyState* KeyState = PlayerInput->GetKeyState(Mapping.Key);
+				if(KeyState && KeyState->bDown)
+				{
+					Mapping.bShouldBeIgnored = true;
+				}
+			}
 
-		// Was this key pressed last frame? If so, then we need to mark it to be ignored by the PlayerInput
-		// until it is released to avoid re-processing a triggered event when it. This is only a problem if
-		// the key was in the old mapping and the new one
-		if(bIgnoreAllPressedKeysUntilReleaseOnRebuild && Mapping.Key.IsDigital() && WasInOldMapping(Mapping.Key))
-		{				
-			const FKeyState* KeyState = PlayerInput->GetKeyState(Mapping.Key);
-			if(KeyState && KeyState->bDown)
+			// Retain old mapping trigger/modifier state for identical key -> action mappings.
+			TArray<FEnhancedActionKeyMapping>::SizeType Idx = OldMappings.IndexOfByPredicate([&Mapping](const FEnhancedActionKeyMapping& Other) {return Mapping == Other; });
+			if (Idx != INDEX_NONE)
 			{
-				Mapping.bShouldBeIgnored = true;
+				Mapping = MoveTemp(OldMappings[Idx]);
+				OldMappings.RemoveAtSwap(Idx);
 			}
 		}
-
-		// Retain old mapping trigger/modifier state for identical key -> action mappings.
-		TArray<FEnhancedActionKeyMapping>::SizeType Idx = OldMappings.IndexOfByPredicate([&Mapping](const FEnhancedActionKeyMapping& Other) {return Mapping == Other; });
-		if (Idx != INDEX_NONE)
+		for (const UInputAction* Action : RemovedActions)
 		{
-			Mapping = MoveTemp(OldMappings[Idx]);
-			OldMappings.RemoveAtSwap(Idx);
-		}
+			PlayerInput->ActionInstanceData.Remove(Action);
+		}	
 	}
-	for (const UInputAction* Action : RemovedActions)
+
+	// Perform a modifier calculation pass on the default data to initialize values correctly.
+	// We do this at the end to ensure ActionInstanceData is accessible without requiring a tick for new/flushed actions.
+	for (FEnhancedActionKeyMapping& Mapping : PlayerInput->EnhancedActionMappings)
 	{
-		PlayerInput->ActionInstanceData.Remove(Action);
+		PlayerInput->InitializeMappingActionModifiers(Mapping);
 	}
-
-	bMappingRebuildPending = false;
+	
+	MappingRebuildPending = EInputMappingRebuildType::None;
 	bIgnoreAllPressedKeysUntilReleaseOnRebuild = true;
 }
 
diff --git a/Engine/Plugins/Experimental/EnhancedInput/Source/EnhancedInput/Private/EnhancedPlayerInput.cpp b/Engine/Plugins/Experimental/EnhancedInput/Source/EnhancedInput/Private/EnhancedPlayerInput.cpp
--- a/Engine/Plugins/Experimental/EnhancedInput/Source/EnhancedInput/Private/EnhancedPlayerInput.cpp
+++ b/Engine/Plugins/Experimental/EnhancedInput/Source/EnhancedInput/Private/EnhancedPlayerInput.cpp
@@ -3,10 +3,10 @@
 #include "EnhancedPlayerInput.h"
 #include "EnhancedInputComponent.h"
 #include "GameFramework/PlayerController.h"
+#include "InputAction.h"
 #include "InputModifiers.h"
 #include "InputTriggers.h"
 #include "Engine/World.h"
-
 // NOTE: Enum order represents firing priority(lowest to highest) and is important as multiple keys bound to the same action may generate differing trigger event states.
 enum class ETriggerEventInternal : uint8
 {
@@ -110,11 +110,12 @@ void UEnhancedPlayerInput::ProcessActionMappingEvent(const UInputAction* Action,
 	FInputActionInstance& ActionData = FindOrAddActionEventData(Action);
 
 	// Update values and triggers for all actionable mappings each frame
-	ETriggerState TriggerState = ETriggerState::None;
+	FTriggerStateTracker TriggerStateTracker;
 
 	// Reset action data on the first event processed for the action this tick.
 	bool bResetActionData = !ActionsWithEventsThisTick.Contains(Action);
-
+	bool bMappingTriggersApplied = false;
+	
 	// If the key state is changing or the key is actuated and being held (and not coming back up this tick) recalculate its value and resulting trigger state.
 	if (KeyEvent != EKeyEvent::None)
 	{
@@ -129,10 +130,11 @@ void UEnhancedPlayerInput::ProcessActionMappingEvent(const UInputAction* Action,
 		FInputActionValue ModifiedValue = ApplyModifiers(Modifiers, FInputActionValue(ValueType, RawKeyValue.Get<FVector>()), DeltaTime);
 		//UE_CLOG(RawKeyValue.GetMagnitudeSq(), LogTemp, Warning, TEXT("Modified %s -> %s"), *RawKeyValue.ToString(), *ModifiedValue.ToString());
 
-		// Derive a trigger state for this mapping using all applicable triggers
-		TriggerState = CalcTriggerState(Triggers, ModifiedValue, DeltaTime);
-		ActionData.bMappingTriggerApplied |= Triggers.Num() > 0;
-
+		// Derive an initial trigger state for this mapping using all applicable triggers
+		ETriggerState CalcedState = TriggerStateTracker.EvaluateTriggers(this, Triggers, ModifiedValue, DeltaTime);
+		// Do this only for no triggers?
+		TriggerStateTracker.SetStateForNoTriggers(ModifiedValue.IsNonZero() ? ETriggerState::Triggered : ETriggerState::None);	
+		bMappingTriggersApplied = Triggers.Num() > 0;
 		// Combine values for active events only, selecting the input with the greatest magnitude for each component in each tick.
 		if(ModifiedValue.GetMagnitudeSq())
 		{
@@ -150,7 +152,9 @@ void UEnhancedPlayerInput::ProcessActionMappingEvent(const UInputAction* Action,
 		}
 	}
 
-	ActionData.MappingTriggerState = FMath::Max(ActionData.MappingTriggerState, TriggerState);
+	// Retain the most interesting/triggered tracker.
+	ActionData.TriggerStateTracker = FMath::Max(ActionData.TriggerStateTracker, TriggerStateTracker);
+	ActionData.TriggerStateTracker.SetMappingTriggerApplied(bMappingTriggersApplied);
 }
 
 void UEnhancedPlayerInput::InjectInputForAction(const UInputAction* Action, FInputActionValue RawValue, const TArray<UInputModifier*>& Modifiers, const TArray<UInputTrigger*>& Triggers)
@@ -299,16 +303,23 @@ void UEnhancedPlayerInput::ProcessInputStack(const TArray<UInputComponent*>& Inp
 
 		if (ActionsWithEventsThisTick.Contains(Action))
 		{
-			// Apply modifiers
+			// Apply action modifiers
+			FInputActionValue RawValue = ActionData.Value; 
 			ActionData.Value = ApplyModifiers(ActionData.Modifiers, ActionData.Value, NonDilatedDeltaTime);
 
-			// Evaluate triggers
-			TriggerState = CalcTriggerState(ActionData.Triggers, ActionData.Value, NonDilatedDeltaTime);
-
-			// Any mapping triggers applied should limit the final state.
-			TriggerState = ActionData.bMappingTriggerApplied ? FMath::Min(TriggerState, ActionData.MappingTriggerState) : TriggerState;
+			// Update what state to use for this data in the case of there being no triggers, otherwise we can get incorrect triggered
+			// states even if the modified value is Zero
+			if(ActionData.Value.Get<FVector>() != RawValue.Get<FVector>())
+			{
+				ActionData.TriggerStateTracker.SetStateForNoTriggers(ActionData.Value.IsNonZero() ? ETriggerState::Triggered : ETriggerState::None);	
+			}
 
-			// However, if the game is paused invalidate trigger unless the action allows it. We must always call CalcTriggerState to update any internal state, even when paused.
+			ETriggerState PrevState = ActionData.TriggerStateTracker.GetState();
+			// Evaluate action triggers. We must always call EvaluateTriggers to update any internal state, even when paused.
+			TriggerState = ActionData.TriggerStateTracker.EvaluateTriggers(this, ActionData.Triggers, ActionData.Value, NonDilatedDeltaTime);
+			TriggerState = ActionData.TriggerStateTracker.GetMappingTriggerApplied() ? FMath::Min(TriggerState, PrevState) : TriggerState;
+			
+			// However, if the game is paused invalidate trigger unless the action allows it.
 			// TODO: Potential issues with e.g. hold event that's canceled due to pausing, but jumps straight back to its "triggered" state on unpause if the user continues to hold the key.
 			if (bGamePaused && !Action->bTriggerWhenPaused)
 			{
@@ -468,8 +479,7 @@ void UEnhancedPlayerInput::ProcessInputStack(const TArray<UInputComponent*>& Inp
 		}
 
 		// Delay MappingTriggerState reset until here to allow dependent triggers (e.g. chords) access to this tick's values.
-		ActionData.MappingTriggerState = ETriggerState::None;
-		ActionData.bMappingTriggerApplied = false;
+		ActionData.TriggerStateTracker = FTriggerStateTracker();
 	}
 
 	LastFrameTime = CurrentTime;
@@ -551,70 +561,6 @@ FInputActionValue UEnhancedPlayerInput::ApplyModifiers(const TArray<UInputModifi
 	return ModifiedValue;
 }
 
-// Calculate a collective representation of trigger state from all key mapping trigger states
-ETriggerState UEnhancedPlayerInput::CalcTriggerState(const TArray<UInputTrigger*>& Triggers, FInputActionValue ModifiedValue, float DeltaTime) const
-{
-	TRACE_CPUPROFILER_EVENT_SCOPE(EnhPIS_Triggers);
-
-	// Trigger rules by implicit/explicit count:
-	// Implicits == 0, Explicits == 0	- Always fire, unless value is 0.
-	// Implicits == 0, Explicits  > 0	- At least one explict has fired.
-	// Implicits  > 0, Explicits == 0	- All implicits have fired. NOTE: Some implicits firing
-	// Implicits  > 0, Explicits  > 0	- All implicits and at least one explicit have fired.
-	// Blockers							- Override all other triggers to force trigger failure.
-
-	if (Triggers.Num() == 0)
-	{
-		// With no triggers the trigger state is represented directly by ModifiedValue.
-		return ModifiedValue.IsNonZero() ? ETriggerState::Triggered : ETriggerState::None;
-	}
-
-	bool bFoundActiveTrigger = false;	// If any trigger is in an ongoing or triggered state the final state must be at least ongoing (with the exception of blocking triggers!)
-	bool bAnyExplictTriggered = false;
-	bool bFoundExplicit = false;		// If no explicits are found the trigger may fire through implicit testing only. If explicits exist at least one must be met.
-	bool bAllImplicitsTriggered = true;
-	bool bBlocking = false;				// If any trigger is blocking, we can't fire.
-
-	// TODO: Make this more efficient. Split implicit/explicit to allow us to early out on implicit fail/explicit pass?
-	for (UInputTrigger* Trigger : Triggers)
-	{
-		if (!Trigger)
-		{
-			continue;
-		}
-
-		ETriggerState CurrentState = Trigger->UpdateState(this, ModifiedValue, DeltaTime);
-
-		// Automatically update the last value, avoiding the trigger having to track it.
-		Trigger->LastValue = ModifiedValue;
-
-		switch (Trigger->GetTriggerType())
-		{
-		case ETriggerType::Explicit:
-			bFoundExplicit = true;
-			bAnyExplictTriggered |= (CurrentState == ETriggerState::Triggered);
-			bFoundActiveTrigger |= (CurrentState != ETriggerState::None);
-			break;
-		case ETriggerType::Implicit:
-			bAllImplicitsTriggered &= (CurrentState == ETriggerState::Triggered);
-			bFoundActiveTrigger |= (CurrentState != ETriggerState::None);
-			break;
-		case ETriggerType::Blocker:
-			bBlocking |= (CurrentState == ETriggerState::Triggered);
-			// Ongoing blockers don't count as active triggers
-			break;
-		}
-	}
-
-	if (bBlocking)
-	{
-		return ETriggerState::None;
-	}
-
-	bool bTriggered = ((!bFoundExplicit || bAnyExplictTriggered) && bAllImplicitsTriggered);
-	return bTriggered ? ETriggerState::Triggered : (bFoundActiveTrigger ? ETriggerState::Ongoing : ETriggerState::None);
-}
-
 bool UEnhancedPlayerInput::IsKeyHandledByAction(FKey Key) const
 {
 	// Determines if the key event is handled or not.
diff --git a/Engine/Plugins/Experimental/EnhancedInput/Source/EnhancedInput/Private/InputAction.cpp b/Engine/Plugins/Experimental/EnhancedInput/Source/EnhancedInput/Private/InputAction.cpp
--- a/Engine/Plugins/Experimental/EnhancedInput/Source/EnhancedInput/Private/InputAction.cpp
+++ b/Engine/Plugins/Experimental/EnhancedInput/Source/EnhancedInput/Private/InputAction.cpp
@@ -34,6 +34,70 @@ FInputActionInstance::FInputActionInstance(const UInputAction* InSourceAction)
 	}
 }
 
+// Calculate a collective representation of trigger state from evaluations of all triggers in one or more trigger groups.
+ETriggerState FTriggerStateTracker::EvaluateTriggers(const UEnhancedPlayerInput* PlayerInput, const TArray<UInputTrigger*>& Triggers, FInputActionValue ModifiedValue, float DeltaTime)
+{
+	TRACE_CPUPROFILER_EVENT_SCOPE(EnhPIS_Triggers);
+
+	// Note: No early outs permitted (e.g. whilst bBlocking)! All triggers must be evaluated to update their internal state/delta time.
+
+	for (UInputTrigger* Trigger : Triggers)
+	{
+		if (!Trigger)
+		{
+			continue;
+		}
+
+		bEvaluatedTriggers = true;
+
+		ETriggerState CurrentState = Trigger->UpdateState(PlayerInput, ModifiedValue, DeltaTime);
+
+		// Automatically update the last value, avoiding the trigger having to track it.
+		Trigger->LastValue = ModifiedValue;
+
+		switch (Trigger->GetTriggerType())
+		{
+		case ETriggerType::Explicit:
+			bFoundExplicit = true;
+			bAnyExplictTriggered |= (CurrentState == ETriggerState::Triggered);
+			bFoundActiveTrigger |= (CurrentState != ETriggerState::None);
+			break;
+		case ETriggerType::Implicit:
+			bAllImplicitsTriggered &= (CurrentState == ETriggerState::Triggered);
+			bFoundActiveTrigger |= (CurrentState != ETriggerState::None);
+			break;
+		case ETriggerType::Blocker:
+			bBlocking |= (CurrentState == ETriggerState::Triggered);
+			// Ongoing blockers don't count as active triggers
+			break;
+		}
+	}
+
+	return GetState();
+}
+
+ETriggerState FTriggerStateTracker::GetState() const
+{
+	if(!bEvaluatedTriggers)
+	{
+		return NoTriggerState;
+	}
+
+	if (bBlocking)
+	{
+		return ETriggerState::None;
+	}
+
+	bool bTriggered = ((!bFoundExplicit || bAnyExplictTriggered) && bAllImplicitsTriggered);
+	return bTriggered ? ETriggerState::Triggered : (bFoundActiveTrigger ? ETriggerState::Ongoing : ETriggerState::None);
+}
+
+// TODO: Hacky. This is the state we should return if we have evaluated no valid triggers. Set during action evaluation based on final ModifiedValue.
+void FTriggerStateTracker::SetStateForNoTriggers(ETriggerState State)
+{
+	NoTriggerState = State;
+}
+
 #if WITH_EDITOR
 // Record input action property changes for later processing
 
diff --git a/Engine/Plugins/Experimental/EnhancedInput/Source/EnhancedInput/Private/InputTriggers.cpp b/Engine/Plugins/Experimental/EnhancedInput/Source/EnhancedInput/Private/InputTriggers.cpp
--- a/Engine/Plugins/Experimental/EnhancedInput/Source/EnhancedInput/Private/InputTriggers.cpp
+++ b/Engine/Plugins/Experimental/EnhancedInput/Source/EnhancedInput/Private/InputTriggers.cpp
@@ -206,5 +206,5 @@ ETriggerState UInputTriggerChordAction::UpdateState_Implementation(const UEnhanc
 {
 	// Inherit state from the chorded action
 	const FInputActionInstance* EventData = PlayerInput->FindActionInstanceData(ChordAction);
-	return EventData ? EventData->MappingTriggerState : ETriggerState::None;
+	return EventData ? EventData->TriggerStateTracker.GetState() : ETriggerState::None;
 };
\ No newline at end of file
diff --git a/Engine/Plugins/Experimental/EnhancedInput/Source/EnhancedInput/Public/EnhancedInputSubsystemInterface.h b/Engine/Plugins/Experimental/EnhancedInput/Source/EnhancedInput/Public/EnhancedInputSubsystemInterface.h
--- a/Engine/Plugins/Experimental/EnhancedInput/Source/EnhancedInput/Public/EnhancedInputSubsystemInterface.h
+++ b/Engine/Plugins/Experimental/EnhancedInput/Source/EnhancedInput/Public/EnhancedInputSubsystemInterface.h
@@ -23,6 +23,17 @@ class UEnhancedInputSubsystemInterface : public UInterface
 	GENERATED_BODY()
 };
 
+UENUM()
+enum class EInputMappingRebuildType : uint8
+{
+	// No rebuild required.
+	None,
+	// Standard mapping rebuild. Retains existing triggers and modifiers for actions that were previously mapped.
+	Rebuild,
+	// If you have made changes to the triggers/modifiers associated with a UInputAction that was previously mapped a flush is required to reset the tracked data for that action.
+	RebuildWithFlush,
+};
+
 /** Passed in as params for Adding/Remove input contexts */
 USTRUCT(BlueprintType)
 struct FModifyContextOptions
@@ -124,7 +135,7 @@ public:
 	* @param Options		Options to consider when removing this input mapping context
 	*/
 	UFUNCTION(BlueprintCallable, BlueprintCosmetic, Category = "Input", meta=(AutoCreateRefTerm = "Options"))
-	virtual void RequestRebuildControlMappings(const FModifyContextOptions& Options = FModifyContextOptions());
+	virtual void RequestRebuildControlMappings(const FModifyContextOptions& Options = FModifyContextOptions(), EInputMappingRebuildType RebuildType = EInputMappingRebuildType::Rebuild);
 
 	/**
 	 * Check if a key mapping is safe to add to a given mapping context within the set of active contexts currently applied to the player controller.
@@ -188,6 +199,8 @@ private:
 	TMap<TWeakObjectPtr<const UInputAction>, FInputActionValue> ForcedActions;
 	TMap<FKey, FInputActionValue> ForcedKeys;
 
+	EInputMappingRebuildType MappingRebuildPending = EInputMappingRebuildType::None;
+
 	/**
 	 * A flag that will be set when adding/removing a mapping context.
 	 *
diff --git a/Engine/Plugins/Experimental/EnhancedInput/Source/EnhancedInput/Public/EnhancedPlayerInput.h b/Engine/Plugins/Experimental/EnhancedInput/Source/EnhancedInput/Public/EnhancedPlayerInput.h
--- a/Engine/Plugins/Experimental/EnhancedInput/Source/EnhancedInput/Public/EnhancedPlayerInput.h
+++ b/Engine/Plugins/Experimental/EnhancedInput/Source/EnhancedInput/Public/EnhancedPlayerInput.h
@@ -62,7 +62,6 @@ private:
 	void InitializeMappingActionModifiers(const FEnhancedActionKeyMapping& Mapping);
 
 	FInputActionValue ApplyModifiers(const TArray<UInputModifier*>& Modifiers, FInputActionValue RawValue, float DeltaTime) const;						// Pre-modified (raw) value
-	ETriggerState CalcTriggerState(const TArray<UInputTrigger*>& Triggers, FInputActionValue ModifiedValue, float DeltaTime) const;						// Post-modified value
 	ETriggerEventInternal GetTriggerStateChangeEvent(ETriggerState LastTriggerState, ETriggerState NewTriggerState) const;
 	ETriggerEvent ConvertInternalTriggerEvent(ETriggerEventInternal Event) const;	// Collapse a detailed internal trigger event into a friendly representation
 	void ProcessActionMappingEvent(const UInputAction* Action, float DeltaTime, bool bGamePaused, FInputActionValue RawValue, EKeyEvent KeyEvent, const TArray<UInputModifier*>& Modifiers, const TArray<UInputTrigger*>& Triggers);
diff --git a/Engine/Plugins/Experimental/EnhancedInput/Source/EnhancedInput/Public/InputAction.h b/Engine/Plugins/Experimental/EnhancedInput/Source/EnhancedInput/Public/InputAction.h
--- a/Engine/Plugins/Experimental/EnhancedInput/Source/EnhancedInput/Public/InputAction.h
+++ b/Engine/Plugins/Experimental/EnhancedInput/Source/EnhancedInput/Public/InputAction.h
@@ -65,6 +65,43 @@ public:
 	TArray<UInputModifier*> Modifiers;
 };
 
+// Calculate a collective representation of trigger state from evaluations of all triggers in one or more trigger groups.
+struct FTriggerStateTracker
+{
+	// Trigger rules by evaluated trigger count:
+	// Implicits == 0, Explicits == 0	- Always fire, unless value is 0.
+	// Implicits == 0, Explicits  > 0	- At least one explict has fired.
+	// Implicits  > 0, Explicits == 0	- All implicits have fired.
+	// Implicits  > 0, Explicits  > 0	- All implicits and at least one explicit have fired.
+	// Blockers   > 0					- Override all other triggers to force trigger failure.
+
+	// Add a group of triggers to the evaluated state, returning the new trigger state.
+	ETriggerState EvaluateTriggers(const UEnhancedPlayerInput* PlayerInput, const TArray<UInputTrigger*>& Triggers, FInputActionValue ModifiedValue, float DeltaTime);
+
+	ETriggerState GetState() const;
+
+	void SetMappingTriggerApplied(bool bNewVal) { bMappingTriggerApplied = bNewVal; }
+	bool GetMappingTriggerApplied() const { return bMappingTriggerApplied; }
+
+	bool operator>=(const FTriggerStateTracker& Other) const { return GetState() >= Other.GetState(); }
+
+	// TODO: Hacky. This is the state we should return if we have evaluated no valid triggers. Set during action evaluation based on final ModifiedValue.
+	void SetStateForNoTriggers(ETriggerState State);
+
+private:
+	ETriggerState NoTriggerState = ETriggerState::None;
+
+	bool bEvaluatedInput = false;		// Non-zero input value was provided at some point in the evaluation
+	bool bEvaluatedTriggers = false;	// At least one valid trigger was evaluated
+	bool bFoundActiveTrigger = false;	// If any trigger is in an ongoing or triggered state the final state must be at least ongoing (with the exception of blocking triggers!)
+	bool bAnyExplictTriggered = false;
+	bool bFoundExplicit = false;		// If no explicits are found the trigger may fire through implicit testing only. If explicits exist at least one must be met.
+	bool bAllImplicitsTriggered = true;
+	bool bBlocking = false;				// If any trigger is blocking, we can't fire.
+	bool bMappingTriggerApplied = false; // Set to true when an actionmapping is processed and triggers were found
+};
+
+
 // Run time queryable action instance
 // Generated from UInputAction templates above
 USTRUCT(BlueprintType)
@@ -83,9 +120,8 @@ private:
 
 	// Internal trigger states
 	ETriggerState LastTriggerState = ETriggerState::None;
-	ETriggerState MappingTriggerState = ETriggerState::None;
+	FTriggerStateTracker TriggerStateTracker;
 	ETriggerEventInternal TriggerEventInternal = ETriggerEventInternal(0);	// TODO: Expose access to ETriggerEventInternal?
-	bool bMappingTriggerApplied = false;
 
 protected:
 	// TODO: Just hold a duplicate of the UInputAction in here?
