Mark the ETriggerEvent enum as a BlueprintType so that it can passed around to different BP event listeners if you have a native C++ implementation

#rnx

diff --git a/Engine/Plugins/Experimental/EnhancedInput/Source/EnhancedInput/Public/InputTriggers.h b/Engine/Plugins/Experimental/EnhancedInput/Source/EnhancedInput/Public/InputTriggers.h
--- a/Engine/Plugins/Experimental/EnhancedInput/Source/EnhancedInput/Public/InputTriggers.h
+++ b/Engine/Plugins/Experimental/EnhancedInput/Source/EnhancedInput/Public/InputTriggers.h
@@ -28,7 +28,7 @@ enum class ETriggerState : uint8
 /**
 * Trigger events are the Action's interpretation of all Trigger State transitions that occurred for the action in the last tick
 */
-UENUM()
+UENUM(BlueprintType)
 enum class ETriggerEvent : uint8
 {
 	// No significant trigger state changes occurred and there are no active device inputs
