Renamed Experimental Copy of FActionMappingsNodeBuilderas in teh EhancedInput plugin this symbol already exists in the engine

diff --git a/Engine/Plugins/Experimental/EnhancedInput/Source/InputEditor/Private/ActionMappingDetails.cpp b/Engine/Plugins/Experimental/EnhancedInput/Source/InputEditor/Private/ActionMappingDetails.cpp
--- a/Engine/Plugins/Experimental/EnhancedInput/Source/InputEditor/Private/ActionMappingDetails.cpp
+++ b/Engine/Plugins/Experimental/EnhancedInput/Source/InputEditor/Private/ActionMappingDetails.cpp
@@ -20,13 +20,13 @@
 
 // TODO: This is derived from (and will eventually replace) InputSettingsDetails.cpp
 
-FActionMappingsNodeBuilder::FActionMappingsNodeBuilder(IDetailLayoutBuilder* InDetailLayoutBuilder, const TSharedPtr<IPropertyHandle>& InPropertyHandle)
+FActionMappingsNodeBuilderEx::FActionMappingsNodeBuilderEx(IDetailLayoutBuilder* InDetailLayoutBuilder, const TSharedPtr<IPropertyHandle>& InPropertyHandle)
 	: DetailLayoutBuilder(InDetailLayoutBuilder)
 	, ActionMappingsPropertyHandle(InPropertyHandle)
 {
 }
 
-void FActionMappingsNodeBuilder::Tick(float DeltaTime)
+void FActionMappingsNodeBuilderEx::Tick(float DeltaTime)
 {
 	if (GroupsRequireRebuild())
 	{
@@ -35,17 +35,17 @@ void FActionMappingsNodeBuilder::Tick(float DeltaTime)
 	HandleDelayedGroupExpansion();
 }
 
-void FActionMappingsNodeBuilder::GenerateHeaderRowContent(FDetailWidgetRow& NodeRow)
+void FActionMappingsNodeBuilderEx::GenerateHeaderRowContent(FDetailWidgetRow& NodeRow)
 {
 	TSharedRef<SWidget> AddButton = PropertyCustomizationHelpers::MakeAddButton(
-		FSimpleDelegate::CreateSP(this, &FActionMappingsNodeBuilder::AddActionMappingButton_OnClick),
-		TAttribute<FText>(this, &FActionMappingsNodeBuilder::GetAddNewActionTooltip),
-		TAttribute<bool>(this, &FActionMappingsNodeBuilder::CanAddNewActionMapping));
+		FSimpleDelegate::CreateSP(this, &FActionMappingsNodeBuilderEx::AddActionMappingButton_OnClick),
+		TAttribute<FText>(this, &FActionMappingsNodeBuilderEx::GetAddNewActionTooltip),
+		TAttribute<bool>(this, &FActionMappingsNodeBuilderEx::CanAddNewActionMapping));
 
-	TSharedRef<SWidget> ClearButton = PropertyCustomizationHelpers::MakeEmptyButton(FSimpleDelegate::CreateSP(this, &FActionMappingsNodeBuilder::ClearActionMappingButton_OnClick),
+	TSharedRef<SWidget> ClearButton = PropertyCustomizationHelpers::MakeEmptyButton(FSimpleDelegate::CreateSP(this, &FActionMappingsNodeBuilderEx::ClearActionMappingButton_OnClick),
 		LOCTEXT("ClearActionMappingToolTip", "Removes all Action Mappings"));
 
-	FSimpleDelegate RebuildChildrenDelegate = FSimpleDelegate::CreateSP(this, &FActionMappingsNodeBuilder::RebuildChildren);
+	FSimpleDelegate RebuildChildrenDelegate = FSimpleDelegate::CreateSP(this, &FActionMappingsNodeBuilderEx::RebuildChildren);
 	ActionMappingsPropertyHandle->SetOnPropertyValueChanged(RebuildChildrenDelegate);
 	ActionMappingsPropertyHandle->AsArray()->SetOnNumElementsChanged(RebuildChildrenDelegate);
 
@@ -77,7 +77,7 @@ void FActionMappingsNodeBuilder::GenerateHeaderRowContent(FDetailWidgetRow& Node
 	];
 }
 
-bool FActionMappingsNodeBuilder::CanAddNewActionMapping() const
+bool FActionMappingsNodeBuilderEx::CanAddNewActionMapping() const
 {
 	// If the last action mapping the user has added is still null, then do not allow adding another one
 	TSharedPtr<IPropertyHandleArray> ActionMappingsArrayHandle = ActionMappingsPropertyHandle->AsArray();
@@ -97,7 +97,7 @@ bool FActionMappingsNodeBuilder::CanAddNewActionMapping() const
 	return true;
 }
 
-FText FActionMappingsNodeBuilder::GetAddNewActionTooltip() const
+FText FActionMappingsNodeBuilderEx::GetAddNewActionTooltip() const
 {
 	if(CanAddNewActionMapping())
 	{
@@ -109,7 +109,7 @@ FText FActionMappingsNodeBuilder::GetAddNewActionTooltip() const
 	}
 }
 
-void FActionMappingsNodeBuilder::GenerateChildContent(IDetailChildrenBuilder& ChildrenBuilder)
+void FActionMappingsNodeBuilderEx::GenerateChildContent(IDetailChildrenBuilder& ChildrenBuilder)
 {
 	RebuildGroupedMappings();
 
@@ -123,10 +123,10 @@ void FActionMappingsNodeBuilder::GenerateChildContent(IDetailChildrenBuilder& Ch
 		IDetailGroup& ActionMappingGroup = ChildrenBuilder.AddGroup(GroupName, FText::FromString(MappingSet.SharedAction->GetPathName()));
 		MappingSet.DetailGroup = &ActionMappingGroup;
 
-		TSharedRef<SWidget> AddButton = PropertyCustomizationHelpers::MakeAddButton(FSimpleDelegate::CreateSP(this, &FActionMappingsNodeBuilder::AddActionMappingToGroupButton_OnClick, MappingSet),
+		TSharedRef<SWidget> AddButton = PropertyCustomizationHelpers::MakeAddButton(FSimpleDelegate::CreateSP(this, &FActionMappingsNodeBuilderEx::AddActionMappingToGroupButton_OnClick, MappingSet),
 			LOCTEXT("AddActionMappingToGroupToolTip", "Add a control binding to the Action Mapping"));
 
-		TSharedRef<SWidget> RemoveButton = PropertyCustomizationHelpers::MakeDeleteButton(FSimpleDelegate::CreateSP(this, &FActionMappingsNodeBuilder::RemoveActionMappingGroupButton_OnClick, MappingSet),
+		TSharedRef<SWidget> RemoveButton = PropertyCustomizationHelpers::MakeDeleteButton(FSimpleDelegate::CreateSP(this, &FActionMappingsNodeBuilderEx::RemoveActionMappingGroupButton_OnClick, MappingSet),
 			LOCTEXT("RemoveActionMappingGroupToolTip", "Remove the Action Mapping Group"));
 
 		ActionMappingGroup.HeaderRow()
@@ -142,7 +142,7 @@ void FActionMappingsNodeBuilder::GenerateChildContent(IDetailChildrenBuilder& Ch
 					.AllowedClass(UInputAction::StaticClass())
 					.ObjectPath(MappingSet.SharedAction ? MappingSet.SharedAction->GetPathName() : FString())
 					.DisplayUseSelected(true)
-					.OnObjectChanged(this, &FActionMappingsNodeBuilder::OnActionMappingActionChanged, MappingSet)
+					.OnObjectChanged(this, &FActionMappingsNodeBuilderEx::OnActionMappingActionChanged, MappingSet)
 				]
 			]
 			+ SHorizontalBox::Slot()
@@ -170,7 +170,7 @@ void FActionMappingsNodeBuilder::GenerateChildContent(IDetailChildrenBuilder& Ch
 	}
 }
 
-void FActionMappingsNodeBuilder::AddActionMappingButton_OnClick()
+void FActionMappingsNodeBuilderEx::AddActionMappingButton_OnClick()
 {
 	static const FName BaseActionMappingName(*LOCTEXT("NewActionMappingName", "NewActionMapping").ToString());
 	static int32 NewMappingCount = 0;
@@ -192,12 +192,12 @@ void FActionMappingsNodeBuilder::AddActionMappingButton_OnClick()
 	}
 }
 
-void FActionMappingsNodeBuilder::ClearActionMappingButton_OnClick()
+void FActionMappingsNodeBuilderEx::ClearActionMappingButton_OnClick()
 {
 	ActionMappingsPropertyHandle->AsArray()->EmptyArray();
 }
 
-void FActionMappingsNodeBuilder::OnActionMappingActionChanged(const FAssetData& AssetData, const FMappingSet MappingSet)
+void FActionMappingsNodeBuilderEx::OnActionMappingActionChanged(const FAssetData& AssetData, const FMappingSet MappingSet)
 {
 	const FScopedTransaction Transaction(LOCTEXT("SwitchActionMapping_Transaction", "Switch Action Mapping"));
 
@@ -226,7 +226,7 @@ void FActionMappingsNodeBuilder::OnActionMappingActionChanged(const FAssetData&
 	}
 }
 
-void FActionMappingsNodeBuilder::AddActionMappingToGroupButton_OnClick(const FMappingSet MappingSet)
+void FActionMappingsNodeBuilderEx::AddActionMappingToGroupButton_OnClick(const FMappingSet MappingSet)
 {
 	const FScopedTransaction Transaction(LOCTEXT("AddActionMappingToGroup_Transaction", "Add a control binding to the Action Mapping"));
 
@@ -246,7 +246,7 @@ void FActionMappingsNodeBuilder::AddActionMappingToGroupButton_OnClick(const FMa
 	}
 }
 
-void FActionMappingsNodeBuilder::RemoveActionMappingGroupButton_OnClick(const FMappingSet MappingSet)
+void FActionMappingsNodeBuilderEx::RemoveActionMappingGroupButton_OnClick(const FMappingSet MappingSet)
 {
 	const FScopedTransaction Transaction(LOCTEXT("RemoveActionMappingGroup_Transaction", "Remove Action Mapping and all control bindings"));
 
@@ -265,7 +265,7 @@ void FActionMappingsNodeBuilder::RemoveActionMappingGroupButton_OnClick(const FM
 	}
 }
 
-bool FActionMappingsNodeBuilder::GroupsRequireRebuild() const
+bool FActionMappingsNodeBuilderEx::GroupsRequireRebuild() const
 {
 	for (int32 GroupIndex = 0; GroupIndex < GroupedMappings.Num(); ++GroupIndex)
 	{
@@ -283,7 +283,7 @@ bool FActionMappingsNodeBuilder::GroupsRequireRebuild() const
 	return false;
 }
 
-void FActionMappingsNodeBuilder::RebuildGroupedMappings()
+void FActionMappingsNodeBuilderEx::RebuildGroupedMappings()
 {
 	GroupedMappings.Empty();
 
@@ -319,7 +319,7 @@ void FActionMappingsNodeBuilder::RebuildGroupedMappings()
 	}
 }
 
-void FActionMappingsNodeBuilder::HandleDelayedGroupExpansion()
+void FActionMappingsNodeBuilderEx::HandleDelayedGroupExpansion()
 {
 	if (DelayedGroupExpansionStates.Num() > 0)
 	{
diff --git a/Engine/Plugins/Experimental/EnhancedInput/Source/InputEditor/Private/ActionMappingDetails.h b/Engine/Plugins/Experimental/EnhancedInput/Source/InputEditor/Private/ActionMappingDetails.h
--- a/Engine/Plugins/Experimental/EnhancedInput/Source/InputEditor/Private/ActionMappingDetails.h
+++ b/Engine/Plugins/Experimental/EnhancedInput/Source/InputEditor/Private/ActionMappingDetails.h
@@ -29,10 +29,10 @@ struct FMappingSet
 	TArray<TSharedRef<IPropertyHandle>> Mappings;
 };
 
-class FActionMappingsNodeBuilder : public IDetailCustomNodeBuilder, public TSharedFromThis<FActionMappingsNodeBuilder>
+class FActionMappingsNodeBuilderEx : public IDetailCustomNodeBuilder, public TSharedFromThis<FActionMappingsNodeBuilderEx>
 {
 public:
-	FActionMappingsNodeBuilder(IDetailLayoutBuilder* InDetailLayoutBuilder, const TSharedPtr<IPropertyHandle>& InPropertyHandle);
+	FActionMappingsNodeBuilderEx(IDetailLayoutBuilder* InDetailLayoutBuilder, const TSharedPtr<IPropertyHandle>& InPropertyHandle);
 
 	/** IDetailCustomNodeBuilder interface */
 	virtual void SetOnRebuildChildren(FSimpleDelegate InOnRebuildChildren) override { OnRebuildChildren = InOnRebuildChildren; }
diff --git a/Engine/Plugins/Experimental/EnhancedInput/Source/InputEditor/Private/InputCustomizations.cpp b/Engine/Plugins/Experimental/EnhancedInput/Source/InputEditor/Private/InputCustomizations.cpp
--- a/Engine/Plugins/Experimental/EnhancedInput/Source/InputEditor/Private/InputCustomizations.cpp
+++ b/Engine/Plugins/Experimental/EnhancedInput/Source/InputEditor/Private/InputCustomizations.cpp
@@ -25,7 +25,7 @@ void FInputContextDetails::CustomizeDetails(class IDetailLayoutBuilder& DetailBu
 	ActionMappingsPropertyHandle->MarkHiddenByCustomization();
 
 	IDetailCategoryBuilder& MappingsDetailCategoryBuilder = DetailBuilder.EditCategory(ActionMappingsPropertyHandle->GetDefaultCategoryName());
-	const TSharedRef<FActionMappingsNodeBuilder> ActionMappingsBuilder = MakeShareable(new FActionMappingsNodeBuilder(&DetailBuilder, ActionMappingsPropertyHandle));
+	const TSharedRef<FActionMappingsNodeBuilderEx> ActionMappingsBuilder = MakeShareable(new FActionMappingsNodeBuilderEx(&DetailBuilder, ActionMappingsPropertyHandle));
 	MappingsDetailCategoryBuilder.AddCustomBuilder(ActionMappingsBuilder);
 }
 
