Update the comments on some input triggers to make sure their properties tooltips are correctly displayed in the editor

#rnx

diff --git a/Engine/Plugins/Experimental/EnhancedInput/Source/EnhancedInput/Public/InputModifiers.h b/Engine/Plugins/Experimental/EnhancedInput/Source/EnhancedInput/Public/InputModifiers.h
--- a/Engine/Plugins/Experimental/EnhancedInput/Source/EnhancedInput/Public/InputModifiers.h
+++ b/Engine/Plugins/Experimental/EnhancedInput/Source/EnhancedInput/Public/InputModifiers.h
@@ -168,14 +168,18 @@ protected:
 
 	// TODO: Smoothing variants. Configuration options. e.g. smooth over a set time/frame count.
 
-	float ZeroTime = 0.f; /** How long input has been zero. */
+	/** How long input has been zero. */
+	float ZeroTime = 0.f;
 
-	FInputActionValue AverageValue; /** Current average input/sample */
+	/** Current average input/sample */
+	FInputActionValue AverageValue;
 
-	int32 Samples = 0; /** Number of samples since input  has been zero */
+	/** Number of samples since input  has been zero */
+	int32 Samples = 0;
 
 #define SMOOTH_TOTAL_SAMPLE_TIME_DEFAULT (0.0083f)
-	float TotalSampleTime = SMOOTH_TOTAL_SAMPLE_TIME_DEFAULT;	/** Input sampling total time.  */
+	/** Input sampling total time. */
+	float TotalSampleTime = SMOOTH_TOTAL_SAMPLE_TIME_DEFAULT;
 
 	virtual FInputActionValue ModifyRaw_Implementation(const UEnhancedPlayerInput* PlayerInput, FInputActionValue CurrentValue, float DeltaTime) override;
 };
@@ -295,8 +299,9 @@ class UInputModifierSwizzleAxis final : public UInputModifier
 
 public:
 
+	// Default to XY swap, useful for binding 1D inputs to the Y axis.
 	UPROPERTY(EditAnywhere, BlueprintReadWrite, Category = Settings)
-	EInputAxisSwizzle Order = EInputAxisSwizzle::YXZ;	// Default to XY swap, useful for binding 1D inputs to the Y axis.
+	EInputAxisSwizzle Order = EInputAxisSwizzle::YXZ;
 
 protected:
 	virtual FInputActionValue ModifyRaw_Implementation(const UEnhancedPlayerInput* PlayerInput, FInputActionValue CurrentValue, float DeltaTime) override;
diff --git a/Engine/Plugins/Experimental/EnhancedInput/Source/EnhancedInput/Public/InputTriggers.h b/Engine/Plugins/Experimental/EnhancedInput/Source/EnhancedInput/Public/InputTriggers.h
--- a/Engine/Plugins/Experimental/EnhancedInput/Source/EnhancedInput/Public/InputTriggers.h
+++ b/Engine/Plugins/Experimental/EnhancedInput/Source/EnhancedInput/Public/InputTriggers.h
@@ -115,8 +115,9 @@ protected:
 	virtual ETriggerState UpdateState_Implementation(const UEnhancedPlayerInput* PlayerInput, FInputActionValue ModifiedValue, float DeltaTime);
 
 public:
+	// Point at which this trigger fires
 	UPROPERTY(EditAnywhere, Config, BlueprintReadWrite, Category = "Trigger Settings")
-	float ActuationThreshold = 0.5f;	// Point at which this trigger fires
+	float ActuationThreshold = 0.5f;
 
 	/*
 	* Is the value passed in sufficiently large to be of interest to the trigger.
@@ -171,8 +172,9 @@ class UInputTriggerTimedBase : public UInputTrigger
 
 protected:
 
+	// How long have we been actuating this trigger?
 	UPROPERTY(BlueprintReadWrite, Category = "Trigger Settings")
-	float HeldDuration = 0.0f;			// How long have we been actuating this trigger? // TODO: Annoying given the action mapping is already tracking this.
+	float HeldDuration = 0.0f;			 // TODO: Annoying given the action mapping is already tracking this.
 
 	// Transitions to Ongoing on actuation. Never triggers.
 	virtual ETriggerState UpdateState_Implementation(const UEnhancedPlayerInput* PlayerInput, FInputActionValue ModifiedValue, float DeltaTime) override;
@@ -367,8 +369,9 @@ public:
 
 // Chorded actions
 
-/** UInputTriggerChordAction
-	Applies a chord action that must be triggering for this trigger's action to trigger
+/**
+ * UInputTriggerChordAction
+ * Applies a chord action that must be triggering for this trigger's action to trigger
 */
 UCLASS(NotBlueprintable, MinimalAPI, meta = (DisplayName = "Chorded Action", NotInputConfigurable = "true"))
 class UInputTriggerChordAction : public UInputTrigger
@@ -382,6 +385,7 @@ protected:
 	virtual ETriggerState UpdateState_Implementation(const UEnhancedPlayerInput* PlayerInput, FInputActionValue ModifiedValue, float DeltaTime) override;
 public:
 
+	// The action that must be triggering for this trigger's action to trigger
 	UPROPERTY(EditInstanceOnly, BlueprintReadWrite, Category = "Trigger Settings", meta = (DisplayThumbnail = "false"))
 	const class UInputAction* ChordAction = nullptr;
 };
