Do not allow blueprints to bind to gesture keys if it is not enabled in the settings

diff --git a/Engine/Plugins/Experimental/EnhancedInput/Source/InputBlueprintNodes/Private/K2Node_InputDebugKey.cpp b/Engine/Plugins/Experimental/EnhancedInput/Source/InputBlueprintNodes/Private/K2Node_InputDebugKey.cpp
--- a/Engine/Plugins/Experimental/EnhancedInput/Source/InputBlueprintNodes/Private/K2Node_InputDebugKey.cpp
+++ b/Engine/Plugins/Experimental/EnhancedInput/Source/InputBlueprintNodes/Private/K2Node_InputDebugKey.cpp
@@ -13,6 +13,8 @@
 #include "EditorCategoryUtils.h"
 #include "BlueprintActionDatabaseRegistrar.h"
 #include "InputActionValue.h"
+#include "GameFramework/InputSettings.h"
+#include "Editor.h"								// for FEditorDelegates::OnEnableGestureRecognizerChanged
 
 #define LOCTEXT_NAMESPACE "UK2Node_InputDebugKey"
 
@@ -355,12 +357,19 @@ void UK2Node_InputDebugKey::GetMenuActions(FBlueprintActionDatabaseRegistrar& Ac
 		InputNode->InputKey = Key;
 	};
 
+	auto RefreshClassActions = []()
+	{
+		FBlueprintActionDatabase::Get().RefreshClassActions(StaticClass());
+	};
+
 	// actions get registered under specific object-keys; the idea is that
 	// actions might have to be updated (or deleted) if their object-key is
 	// mutated (or removed)... here we use the node's class (so if the node
 	// type disappears, then the action should go with it)
 	UClass* ActionKey = GetClass();
 
+	const bool bAllowGestures = GetDefault<UInputSettings>()->bEnableGestureRecognizer;
+	
 	// to keep from needlessly instantiating a UBlueprintNodeSpawner (and
 	// iterating over keys), first check to make sure that the registrar is
 	// looking for actions of this type (could be regenerating actions for a
@@ -368,10 +377,21 @@ void UK2Node_InputDebugKey::GetMenuActions(FBlueprintActionDatabaseRegistrar& Ac
 	// corresponding to that asset)
 	if (ActionRegistrar.IsOpenForRegistration(ActionKey))
 	{
+		// Refresh the action database of this node when the input settings are changed
+		static bool bRegisterOnce = true;
+		if (bRegisterOnce)
+		{
+			bRegisterOnce = false;
+			FEditorDelegates::OnEnableGestureRecognizerChanged.AddStatic(RefreshClassActions);
+		}
+		
 		for (const FKey& Key : AllKeys)
 		{
-			// AnyKey is not supported as a debug key
-			if (!Key.IsBindableInBlueprints() || Key == EKeys::AnyKey)
+			// Do not show gesture keys if they are not enabled in the settings
+			const bool bInvalidGestureKey = !bAllowGestures && Key.IsGesture();
+			
+			// AnyKey is not supported as a debug key, gestures can only be used if they are turned on in the settings
+			if (!Key.IsBindableInBlueprints() || Key == EKeys::AnyKey || bInvalidGestureKey)
 			{
 				continue;
 			}
diff --git a/Engine/Source/Editor/BlueprintGraph/Private/K2Node_InputKey.cpp b/Engine/Source/Editor/BlueprintGraph/Private/K2Node_InputKey.cpp
--- a/Engine/Source/Editor/BlueprintGraph/Private/K2Node_InputKey.cpp
+++ b/Engine/Source/Editor/BlueprintGraph/Private/K2Node_InputKey.cpp
@@ -12,6 +12,8 @@
 #include "BlueprintNodeSpawner.h"
 #include "EditorCategoryUtils.h"
 #include "BlueprintActionDatabaseRegistrar.h"
+#include "GameFramework/InputSettings.h"
+#include "Editor.h"								// for FEditorDelegates::OnEnableGestureRecognizerChanged
 
 #define LOCTEXT_NAMESPACE "UK2Node_InputKey"
 
@@ -372,11 +374,18 @@ void UK2Node_InputKey::GetMenuActions(FBlueprintActionDatabaseRegistrar& ActionR
 		InputNode->InputKey = Key;
 	};
 
+	auto RefreshClassActions = []()
+	{
+		FBlueprintActionDatabase::Get().RefreshClassActions(StaticClass());
+	};
+
 	// actions get registered under specific object-keys; the idea is that
 	// actions might have to be updated (or deleted) if their object-key is
 	// mutated (or removed)... here we use the node's class (so if the node
 	// type disappears, then the action should go with it)
 	UClass* ActionKey = GetClass();
+	
+	const bool bAllowGestures = GetDefault<UInputSettings>()->bEnableGestureRecognizer;
 
 	// to keep from needlessly instantiating a UBlueprintNodeSpawner (and
 	// iterating over keys), first check to make sure that the registrar is
@@ -385,10 +394,20 @@ void UK2Node_InputKey::GetMenuActions(FBlueprintActionDatabaseRegistrar& ActionR
 	// corresponding to that asset)
 	if (ActionRegistrar.IsOpenForRegistration(ActionKey))
 	{
+		// Refresh the action database of this node when the input settings are changed
+		static bool bRegisterOnce = true;
+		if (bRegisterOnce)
+		{
+			bRegisterOnce = false;
+			FEditorDelegates::OnEnableGestureRecognizerChanged.AddStatic(RefreshClassActions);
+		}
+		
 		for (const FKey& Key : AllKeys)
 		{
+			// Do not show gesture keys if they are not enabled in the settings
+			const bool bInvalidGestureKey = !bAllowGestures && Key.IsGesture();
 			// these will be handled by UK2Node_GetInputAxisKeyValue and UK2Node_GetInputVectorAxisValue respectively
-			if (!Key.IsBindableInBlueprints() || Key.IsAnalog())
+			if (!Key.IsBindableInBlueprints() || Key.IsAnalog() || bInvalidGestureKey)
 			{
 				continue;
 			}
diff --git a/Engine/Source/Editor/UnrealEd/Private/Editor.cpp b/Engine/Source/Editor/UnrealEd/Private/Editor.cpp
--- a/Engine/Source/Editor/UnrealEd/Private/Editor.cpp
+++ b/Engine/Source/Editor/UnrealEd/Private/Editor.cpp
@@ -143,6 +143,7 @@ FEditorDelegates::FOnAssetsAddExtraObjectsToDelete		FEditorDelegates::OnAssetsAd
 FEditorDelegates::FOnAssetsPreDelete					FEditorDelegates::OnAssetsPreDelete;
 FEditorDelegates::FOnAssetsDeleted						FEditorDelegates::OnAssetsDeleted;
 FEditorDelegates::FOnAssetDragStarted					FEditorDelegates::OnAssetDragStarted;
+FSimpleMulticastDelegate								FEditorDelegates::OnEnableGestureRecognizerChanged;
 FSimpleMulticastDelegate								FEditorDelegates::OnActionAxisMappingsChanged;
 FEditorDelegates::FOnAddLevelToWorld					FEditorDelegates::OnAddLevelToWorld;
 FEditorDelegates::FOnEditCutActorsBegin					FEditorDelegates::OnEditCutActorsBegin;
diff --git a/Engine/Source/Editor/UnrealEd/Public/Editor.h b/Engine/Source/Editor/UnrealEd/Public/Editor.h
--- a/Engine/Source/Editor/UnrealEd/Public/Editor.h
+++ b/Engine/Source/Editor/UnrealEd/Public/Editor.h
@@ -286,6 +286,8 @@ struct UNREALED_API FEditorDelegates
 	static FOnAssetsDeleted OnAssetsDeleted;
 	/** Called when a user starts dragging something out of content browser (can be multiple assets) */
 	static FOnAssetDragStarted OnAssetDragStarted;
+	/** Called when a user changes the UInputSettings::bEnableGestureRecognizer setting to refresh the available actions. */
+	static FSimpleMulticastDelegate OnEnableGestureRecognizerChanged;
 	/** Called when Action or Axis mappings have been changed */
 	static FSimpleMulticastDelegate OnActionAxisMappingsChanged;
 	/** Called from FEditorUtils::AddLevelToWorld after the level is added successfully to the world. */
diff --git a/Engine/Source/Runtime/Engine/Private/UserInterface/InputSettings.cpp b/Engine/Source/Runtime/Engine/Private/UserInterface/InputSettings.cpp
--- a/Engine/Source/Runtime/Engine/Private/UserInterface/InputSettings.cpp
+++ b/Engine/Source/Runtime/Engine/Private/UserInterface/InputSettings.cpp
@@ -188,6 +188,11 @@ void UInputSettings::PostEditChangeChainProperty(FPropertyChangedChainEvent& Pro
 		ForceRebuildKeymaps();
 		FEditorDelegates::OnActionAxisMappingsChanged.Broadcast();
 	}
+
+	if (MemberPropertyName == GET_MEMBER_NAME_CHECKED(UInputSettings, bEnableGestureRecognizer))
+	{
+		FEditorDelegates::OnEnableGestureRecognizerChanged.Broadcast();
+	}
 }
 
 #endif
diff --git a/Engine/Source/Runtime/InputCore/Classes/InputCoreTypes.h b/Engine/Source/Runtime/InputCore/Classes/InputCoreTypes.h
--- a/Engine/Source/Runtime/InputCore/Classes/InputCoreTypes.h
+++ b/Engine/Source/Runtime/InputCore/Classes/InputCoreTypes.h
@@ -89,6 +89,7 @@ struct INPUTCORE_API FKey
 	bool ShouldUpdateAxisWithoutSamples() const;
 	bool IsBindableToActions() const;
 	bool IsDeprecated() const;
+	bool IsGesture() const;
 	FText GetDisplayName(bool bLongDisplayName = true) const;
 	FString ToString() const;
 	FName GetFName() const;
@@ -157,6 +158,9 @@ struct INPUTCORE_API FKeyDetails
 		ButtonAxis				 = 1 << 10,		// Analog 1D axis emulating a digital button press. E.g. Gamepad right stick up
 		Axis2D					 = 1 << 11,
 
+		// Gesture input types such as Flick, Pinch, and Rotate
+		Gesture					 = 1 << 12,
+
 		// Deprecated. Replace with axis definitions for clarity.
 
 		FloatAxis  UE_DEPRECATED(4.26, "Please use Axis1D instead.") = Axis1D,
@@ -185,6 +189,7 @@ struct INPUTCORE_API FKeyDetails
 	FORCEINLINE bool IsBindableInBlueprints() const { return bIsBindableInBlueprints != 0; }
 	FORCEINLINE bool ShouldUpdateAxisWithoutSamples() const { return bShouldUpdateAxisWithoutSamples != 0; }
 	FORCEINLINE bool IsBindableToActions() const { return bIsBindableToActions != 0; }
+	FORCEINLINE bool IsGesture() const { return bIsGesture != 0; }
 	FORCEINLINE bool IsDeprecated() const { return bIsDeprecated != 0; }
 	FORCEINLINE FName GetMenuCategory() const { return MenuCategory; }
 	FText GetDisplayName(const bool bLongDisplayName = true) const;
@@ -224,6 +229,7 @@ private:
 	uint8 bShouldUpdateAxisWithoutSamples : 1;
 	uint8 bIsBindableToActions : 1;
 	uint8 bIsDeprecated : 1;
+	uint8 bIsGesture : 1;
 	EInputAxisType AxisType;
 
 	TAttribute<FText> LongDisplayName;
diff --git a/Engine/Source/Runtime/InputCore/Private/InputCoreTypes.cpp b/Engine/Source/Runtime/InputCore/Private/InputCoreTypes.cpp
--- a/Engine/Source/Runtime/InputCore/Private/InputCoreTypes.cpp
+++ b/Engine/Source/Runtime/InputCore/Private/InputCoreTypes.cpp
@@ -460,6 +460,7 @@ void FKeyDetails::CommonInit(const uint32 InKeyFlags)
 	bShouldUpdateAxisWithoutSamples = ((InKeyFlags & EKeyFlags::UpdateAxisWithoutSamples) != 0);
 	bIsBindableToActions = ((~InKeyFlags & EKeyFlags::NotActionBindableKey) != 0) && ((~InKeyFlags & EKeyFlags::Deprecated) != 0);
 	bIsDeprecated = ((InKeyFlags & EKeyFlags::Deprecated) != 0);
+	bIsGesture = ((InKeyFlags & EKeyFlags::Gesture) != 0);
 
 	if ((InKeyFlags & EKeyFlags::ButtonAxis) != 0)
 	{
@@ -731,9 +732,9 @@ void EKeys::Initialize()
 	// Gestures
 	AddMenuCategoryDisplayInfo("Gesture", LOCTEXT("GestureSubCateogry", "Gesture"), TEXT("GraphEditor.KeyEvent_16x"));
 
-	AddKey(FKeyDetails(EKeys::Gesture_Pinch, LOCTEXT("Gesture_Pinch", "Pinch"), 0, "Gesture"));
-	AddKey(FKeyDetails(EKeys::Gesture_Flick, LOCTEXT("Gesture_Flick", "Flick"), 0, "Gesture"));
-	AddKey(FKeyDetails(EKeys::Gesture_Rotate, LOCTEXT("Gesture_Rotate", "Rotate"), 0, "Gesture"));
+	AddKey(FKeyDetails(EKeys::Gesture_Pinch, LOCTEXT("Gesture_Pinch", "Pinch"), FKeyDetails::Gesture, "Gesture"));
+	AddKey(FKeyDetails(EKeys::Gesture_Flick, LOCTEXT("Gesture_Flick", "Flick"), FKeyDetails::Gesture, "Gesture"));
+	AddKey(FKeyDetails(EKeys::Gesture_Rotate, LOCTEXT("Gesture_Rotate", "Rotate"), FKeyDetails::Gesture, "Gesture"));
 
 	// PS4-specific
 	AddMenuCategoryDisplayInfo("Special Gamepad", LOCTEXT("SpecialGamepadSubCategory", "SpecialGamepad"), TEXT("GraphEditor.PadEvent_16x"));
@@ -1345,6 +1346,12 @@ bool FKey::IsDeprecated() const
 	return (KeyDetails.IsValid() ? KeyDetails->IsDeprecated() : false);
 }
 
+bool FKey::IsGesture() const
+{
+	ConditionalLookupKeyDetails();
+	return (KeyDetails.IsValid() ? KeyDetails->IsGesture() : false);
+}
+
 FText FKeyDetails::GetDisplayName(const bool bLongDisplayName) const
 {
 	return ((bLongDisplayName || !ShortDisplayName.IsSet()) ? LongDisplayName.Get() : ShortDisplayName.Get());
