Mark enhanced input console commands as ECVF_Cheat

diff --git a/Engine/Plugins/Experimental/EnhancedInput/Source/EnhancedInput/Private/EnhancedInputModule.cpp b/Engine/Plugins/Experimental/EnhancedInput/Source/EnhancedInput/Private/EnhancedInputModule.cpp
--- a/Engine/Plugins/Experimental/EnhancedInput/Source/EnhancedInput/Private/EnhancedInputModule.cpp
+++ b/Engine/Plugins/Experimental/EnhancedInput/Source/EnhancedInput/Private/EnhancedInputModule.cpp
@@ -239,28 +239,28 @@ void FEnhancedInputModule::StartupModule()
 		    TEXT("Input.+action"),
 		    TEXT("Provide the named action with a constant input value each frame"),
 		    FConsoleCommandWithWorldAndArgsDelegate::CreateRaw(this, &FEnhancedInputModule::EnableForcedAction),
-		    ECVF_Default
+		    ECVF_Cheat
 	    ));
 
 		ConsoleCommands.Add(IConsoleManager::Get().RegisterConsoleCommand(
 			TEXT("Input.-action"),
 			TEXT("Stop forcing the named action value each frame"),
 			FConsoleCommandWithWorldAndArgsDelegate::CreateRaw(this, &FEnhancedInputModule::DisableForcedAction),
-			ECVF_Default
+			ECVF_Cheat
 		));
 
 		ConsoleCommands.Add(IConsoleManager::Get().RegisterConsoleCommand(
 			TEXT("Input.+key"),
 			TEXT("Provide the named key with a constant input value each frame"),
 			FConsoleCommandWithWorldAndArgsDelegate::CreateRaw(this, &FEnhancedInputModule::EnableForcedKey),
-			ECVF_Default
+			ECVF_Cheat
 		));
 
 		ConsoleCommands.Add(IConsoleManager::Get().RegisterConsoleCommand(
 			TEXT("Input.-key"),
 			TEXT("Stop forcing the named key each frame"),
 			FConsoleCommandWithWorldAndArgsDelegate::CreateRaw(this, &FEnhancedInputModule::DisableForcedKey),
-			ECVF_Default
+			ECVF_Cheat
 		));
 	}
 }
