Manually set the value of boolean keys that are pressed and released within the same frame. One example of this is when SceneViewport immediately calls IE_Pressed and IE_Released right after each other for mouse scroll wheels. We do not expose the EventAccumulator to the UInputTriggers which means that there is no way for Enhanced Input to know that the key should trigger without the Action Instance Value (which is a vector, and would be 0 when the wheel is released).

diff --git a/Engine/Plugins/Experimental/EnhancedInput/Source/EnhancedInput/Private/EnhancedPlayerInput.cpp b/Engine/Plugins/Experimental/EnhancedInput/Source/EnhancedInput/Private/EnhancedPlayerInput.cpp
--- a/Engine/Plugins/Experimental/EnhancedInput/Source/EnhancedInput/Private/EnhancedPlayerInput.cpp
+++ b/Engine/Plugins/Experimental/EnhancedInput/Source/EnhancedInput/Private/EnhancedPlayerInput.cpp
@@ -167,6 +167,18 @@ void UEnhancedPlayerInput::InjectInputForAction(const UInputAction* Action, FInp
 	InputsInjectedThisTick.FindOrAdd(Action).Injected.Emplace(MoveTemp(Input));
 }
 
+bool UEnhancedPlayerInput::InputKey(const FInputKeyParams& Params)
+{
+	const bool bResult = Super::InputKey(Params);
+
+	if (Params.Key.IsButtonAxis() && Params.Event == IE_Pressed)
+	{
+		KeysPressedThisTick.FindOrAdd(Params.Key, Params.Delta);
+	}
+	
+	return bResult;
+}
+
 void UEnhancedPlayerInput::ProcessInputStack(const TArray<UInputComponent*>& InputComponentStack, const float DeltaTime, const bool bGamePaused)
 {
 	// We need to grab the down states of all keys before calling Super::ProcessInputStack as it will leave bDownPrevious in the same state as bDown (i.e. this frame, not last).
@@ -245,6 +257,15 @@ void UEnhancedPlayerInput::ProcessInputStack(const TArray<UInputComponent*>& Inp
 
 		EKeyEvent KeyEvent = bKeyIsHeld ? EKeyEvent::Held : ((bKeyIsDown || bKeyIsReleased) ? EKeyEvent::Actuated : EKeyEvent::None);
 
+		FVector* PressedThisTickValue = KeysPressedThisTick.Find(Mapping.Key);
+		
+		// For keys that were pressed and released within the same frame, set their RawValue so that
+		// InputTriggers are aware that they have been pressed
+		if(PressedThisTickValue && bKeyIsDown && KeyState->EventCounts[IE_Pressed].Num() && KeyState->EventCounts[IE_Released].Num() && RawKeyValue.IsZero())
+		{
+			RawKeyValue = *PressedThisTickValue;
+		}
+
 		// Perform update
 		ProcessActionMappingEvent(Mapping.Action, NonDilatedDeltaTime, bGamePaused, RawKeyValue, KeyEvent, Mapping.Modifiers, Mapping.Triggers);
 	}
@@ -483,6 +504,7 @@ void UEnhancedPlayerInput::ProcessInputStack(const TArray<UInputComponent*>& Inp
 	}
 
 	LastFrameTime = CurrentTime;
+	KeysPressedThisTick.Reset();
 }
 
 FInputActionValue UEnhancedPlayerInput::GetActionValue(const UInputAction* ForAction) const
diff --git a/Engine/Plugins/Experimental/EnhancedInput/Source/EnhancedInput/Public/EnhancedPlayerInput.h b/Engine/Plugins/Experimental/EnhancedInput/Source/EnhancedInput/Public/EnhancedPlayerInput.h
--- a/Engine/Plugins/Experimental/EnhancedInput/Source/EnhancedInput/Public/EnhancedPlayerInput.h
+++ b/Engine/Plugins/Experimental/EnhancedInput/Source/EnhancedInput/Public/EnhancedPlayerInput.h
@@ -40,6 +40,8 @@ public:
 	// Input simulation via injection. Runs modifiers and triggers delegates as if the input had come through the underlying input system as FKeys. Applies action modifiers and triggers on top.
 	void InjectInputForAction(const UInputAction* Action, FInputActionValue RawValue, const TArray<UInputModifier*>& Modifiers = {}, const TArray<UInputTrigger*>& Triggers = {});
 
+	virtual bool InputKey(const FInputKeyParams& Params) override;
+	
 protected:
 
 	// Applies modifiers and triggers without affecting keys read by the base input system
@@ -91,6 +93,13 @@ private:
 	/** Actions which had actuated events at the last call to ProcessInputStack (held/pressed/released) */
 	TSet<const UInputAction*> ActionsWithEventsThisTick;
 
+	/**
+	 * A map of Keys to the amount they were depressed this frame. This is reset with each call to ProcessInputStack
+	 * and is populated within UEnhancedPlayerInput::InputKey.
+	 */
+	UPROPERTY(Transient)
+	TMap<FKey, FVector> KeysPressedThisTick;
+
 	struct FInjectedInput
 	{
 		FInputActionValue RawValue;
