Keep track of the time that a UInputAction was triggered. This makes it really easy to create an input trigger that is dependant on another.

For example, you could make a trigger not happen until a certain number of seconds after another had evaluated to true.

#rnx

diff --git a/Engine/Plugins/Experimental/EnhancedInput/Source/EnhancedInput/Private/EnhancedPlayerInput.cpp b/Engine/Plugins/Experimental/EnhancedInput/Source/EnhancedInput/Private/EnhancedPlayerInput.cpp
--- a/Engine/Plugins/Experimental/EnhancedInput/Source/EnhancedInput/Private/EnhancedPlayerInput.cpp
+++ b/Engine/Plugins/Experimental/EnhancedInput/Source/EnhancedInput/Private/EnhancedPlayerInput.cpp
@@ -355,6 +355,11 @@ void UEnhancedPlayerInput::ProcessInputStack(const TArray<UInputComponent*>& Inp
 		// Evaluate time per action after establishing the internal trigger state across all mappings
 		ActionData.ElapsedProcessedTime += TriggerState != ETriggerState::None ? NonDilatedDeltaTime : 0.f;
 		ActionData.ElapsedTriggeredTime += (ActionData.TriggerEvent == ETriggerEvent::Triggered) ? NonDilatedDeltaTime : 0.f;
+		// Track the time that this trigger was last used
+		if(TriggerState == ETriggerState::Triggered)
+		{
+			ActionData.LastTriggeredWorldTime = CurrentTime;
+		}
 	}
 
 
diff --git a/Engine/Plugins/Experimental/EnhancedInput/Source/EnhancedInput/Public/InputAction.h b/Engine/Plugins/Experimental/EnhancedInput/Source/EnhancedInput/Public/InputAction.h
--- a/Engine/Plugins/Experimental/EnhancedInput/Source/EnhancedInput/Public/InputAction.h
+++ b/Engine/Plugins/Experimental/EnhancedInput/Source/EnhancedInput/Public/InputAction.h
@@ -148,6 +148,10 @@ protected:
 	UPROPERTY(BlueprintReadOnly, Category = Action)
 	float ElapsedTriggeredTime = 0.f;
 
+	// The last time that this evaluated to a Triggered State
+	UPROPERTY(BlueprintReadOnly, Category = Action)
+	float LastTriggeredWorldTime = 0.0f;
+
 	// Trigger state
 	UPROPERTY(BlueprintReadOnly, Category = Action)
 	ETriggerEvent TriggerEvent = ETriggerEvent::None;
@@ -168,6 +172,9 @@ public:
 	// Time the action has been actively triggered (Triggered only)
 	float GetTriggeredTime() const { return ElapsedTriggeredTime; }
 
+	// Time that this action was last actively triggered
+	float GetLastTriggeredWorldTime() const { return LastTriggeredWorldTime; }
+
 	const TArray<UInputTrigger*>& GetTriggers() const { return Triggers; }
 	const TArray<UInputModifier*>& GetModifiers() const { return Modifiers; }
 
