Add support for player mappable keys to enhanced input via the UPlayerMappableInputConfig data asset. Users can now create a UPlayerMappableInputConfig and use it to hold references to any input mapping contexts for a certain config, and it will allow for easy access to save/load player mappable keys. Each EnhancedKeyMapping has a flag for "IsPlayerMappable" and can be enabled to reveal some metadata for each mapping.

#rnx

diff --git a/Engine/Plugins/Experimental/EnhancedInput/Source/EnhancedInput/Private/EnhancedInputSubsystemInterface.cpp b/Engine/Plugins/Experimental/EnhancedInput/Source/EnhancedInput/Private/EnhancedInputSubsystemInterface.cpp
--- a/Engine/Plugins/Experimental/EnhancedInput/Source/EnhancedInput/Private/EnhancedInputSubsystemInterface.cpp
+++ b/Engine/Plugins/Experimental/EnhancedInput/Source/EnhancedInput/Private/EnhancedInputSubsystemInterface.cpp
@@ -11,6 +11,7 @@
 #include "InputModifiers.h"
 #include "InputTriggers.h"
 #include "UObject/UObjectIterator.h"
+#include "PlayerMappableInputConfig.h"
 
 /* Shared input subsystem functionality.
  * See EnhancedInputSubsystemInterfaceDebug.cpp for debug specific functionality.
@@ -376,6 +377,31 @@ TArray<FKey> IEnhancedInputSubsystemInterface::QueryKeysMappedToAction(const UIn
 	return MappedKeys;
 }
 
+int32 IEnhancedInputSubsystemInterface::AddPlayerMappedKey(const FName MappingName, const FKey NewKey, const FModifyContextOptions& Options)
+{
+	int32 NumMappingsApplied = 0;
+	if (UEnhancedPlayerInput* const PlayerInput = GetPlayerInput())
+	{
+		PlayerMappedSettings.Add(MappingName, NewKey);
+		++NumMappingsApplied;
+	}
+
+	RequestRebuildControlMappings(Options);
+	
+	return NumMappingsApplied;
+}
+
+void IEnhancedInputSubsystemInterface::AddPlayerMappableConfig(const UPlayerMappableInputConfig* Config, const FModifyContextOptions& Options)
+{
+	if(Config)
+	{
+		for(TPair<UInputMappingContext*, int32> Pair : Config->GetMappingContexts())
+		{
+			AddMappingContext(Pair.Key, Pair.Value, Options);
+		}	
+	}
+}
+
 template<typename T>
 void DeepCopyPtrArray(const TArray<T*>& From, TArray<T*>& To)
 {
@@ -499,10 +525,15 @@ void IEnhancedInputSubsystemInterface::RebuildControlMappings()
 		TArray<FKey> ContextAppliedKeys;
 
 		const UInputMappingContext* MappingContext = ContextPair.Key;
-		TArray<FEnhancedActionKeyMapping>  OrderedMappings = ReorderMappings(MappingContext->GetMappings());
+		TArray<FEnhancedActionKeyMapping> OrderedMappings = ReorderMappings(MappingContext->GetMappings());
 
-		for (const FEnhancedActionKeyMapping& Mapping : OrderedMappings)
+		for (FEnhancedActionKeyMapping& Mapping : OrderedMappings)
 		{
+			if (FKey* PlayerKey = PlayerMappedSettings.Find(Mapping.PlayerMappableOptions.Name))
+			{
+				Mapping.Key = *PlayerKey;
+			}
+			
 			if (Mapping.Action && !AppliedKeys.Contains(Mapping.Key))
 			{
 				// TODO: Wasteful query as we've already established chord state within ReorderMappings. Store TOptional bConsumeInput per mapping, allowing override? Query override via delegate?
diff --git a/Engine/Plugins/Experimental/EnhancedInput/Source/EnhancedInput/Private/PlayerMappableInputConfig.cpp b/Engine/Plugins/Experimental/EnhancedInput/Source/EnhancedInput/Private/PlayerMappableInputConfig.cpp
new file mode 100644
--- /dev/null
+++ b/Engine/Plugins/Experimental/EnhancedInput/Source/EnhancedInput/Private/PlayerMappableInputConfig.cpp
@@ -0,0 +1,131 @@
+// Copyright Epic Games, Inc. All Rights Reserved.
+
+#include "PlayerMappableInputConfig.h"
+#include "InputMappingContext.h"
+
+#define LOCTEXT_NAMESPACE "PlayerMappableInputConfig"
+
+UPlayerMappableInputConfig::UPlayerMappableInputConfig(const FObjectInitializer& ObjectInitializer)
+{
+	ResetToDefault();
+}
+
+void UPlayerMappableInputConfig::ResetToDefault()
+{
+}
+
+#if WITH_EDITOR
+
+EDataValidationResult UPlayerMappableInputConfig::IsDataValid(TArray<FText>& ValidationErrors)
+{
+	EDataValidationResult Result = Super::IsDataValid(ValidationErrors);
+
+	if(ConfigName == NAME_None)
+	{
+		FFormatNamedArguments Args;
+		Args.Add(TEXT("AssetPath"), FText::FromString(GetPathName()));
+		const FText NoNameError = FText::Format(LOCTEXT("NullContextError", "'{AssetPath}' does not have a valid ConfigName!"), Args);
+		ValidationErrors.Emplace(NoNameError);
+		
+		Result = CombineDataValidationResults(Result, EDataValidationResult::Invalid);
+	}
+	
+	//TMap<const UInputMappingContext*, int32> InputContexts = Contexts;
+	
+	// Check that every context is valid and all player bindable mappings have names
+	for (const TPair<UInputMappingContext*, int32>& ContextPair : Contexts)
+	{
+		if(const UInputMappingContext* IMC = ContextPair.Key)
+		{
+			for(const FEnhancedActionKeyMapping& Mapping : IMC->GetMappings())
+			{
+				if(Mapping.bIsPlayerMappable && Mapping.PlayerMappableOptions.Name == NAME_None)
+				{
+					Result = CombineDataValidationResults(Result, EDataValidationResult::Invalid);
+					FFormatNamedArguments Args;
+					Args.Add(TEXT("AssetPath"), FText::FromString(GetPathName()));
+					Args.Add(TEXT("ActionName"), FText::FromString(Mapping.Action->GetPathName()));
+
+					const FText NoMappingError = FText::Format(LOCTEXT("NoNameMappingError", "'{AssetPath}' has a player bindable key that has no name!"), Args);
+					ValidationErrors.Emplace(NoMappingError);
+				}
+			}
+		}
+		else
+		{
+			Result = CombineDataValidationResults(Result, EDataValidationResult::Invalid);
+			
+			FFormatNamedArguments Args;
+			Args.Add(TEXT("AssetPath"), FText::FromString(GetPathName()));
+
+			const FText NullContextError = FText::Format(LOCTEXT("NullContextError", "There is a null mapping context in '{AssetPath}'"), Args);
+			ValidationErrors.Emplace(NullContextError);
+		}
+	}
+	
+	return CombineDataValidationResults(Result, EDataValidationResult::Valid);
+}
+
+#endif
+
+void UPlayerMappableInputConfig::ForEachDefaultPlayerMappableKey(TFunctionRef<void(const FEnhancedActionKeyMapping&)> Operation) const
+{
+	for (const TPair<UInputMappingContext*, int32>& ContextPair: Contexts)
+	{
+		if(const UInputMappingContext* IMC = ContextPair.Key)
+		{
+			for(const FEnhancedActionKeyMapping& Mapping : IMC->GetMappings())
+			{
+				if(Mapping.bIsPlayerMappable)
+				{
+					Operation(Mapping);
+				}
+			}
+		}
+	}
+}
+
+TArray<FEnhancedActionKeyMapping> UPlayerMappableInputConfig::GetPlayerMappableKeys() const
+{
+	TArray<FEnhancedActionKeyMapping> OutMappings;
+	auto GatherMappings = [&OutMappings](const FEnhancedActionKeyMapping& Options)
+	{
+		OutMappings.Add(Options);
+	};
+	
+	ForEachDefaultPlayerMappableKey(GatherMappings);
+
+	return OutMappings;
+}
+
+FEnhancedActionKeyMapping UPlayerMappableInputConfig::GetMappingByName(const FName MappingName) const
+{
+	FEnhancedActionKeyMapping OutMapping;
+	
+	TArray<FEnhancedActionKeyMapping> MappableKeys = GetPlayerMappableKeys();
+	FEnhancedActionKeyMapping* ExistingMapping = MappableKeys.FindByPredicate([MappingName](const FEnhancedActionKeyMapping& Mapping) { return Mapping.PlayerMappableOptions.Name == MappingName; });
+
+	if(ExistingMapping)
+	{
+		OutMapping = *ExistingMapping;
+	}
+	return OutMapping;
+}
+
+TArray<FEnhancedActionKeyMapping> UPlayerMappableInputConfig::GetKeysBoundToAction(const UInputAction* InAction) const
+{
+	TArray<FEnhancedActionKeyMapping> OutMappings;
+	auto GatherMappings = [&OutMappings, InAction](const FEnhancedActionKeyMapping& Options)
+	{
+		if(Options.Action == InAction)
+		{
+			OutMappings.Add(Options);
+		}
+	};
+	
+	ForEachDefaultPlayerMappableKey(GatherMappings);
+	
+	return OutMappings;
+}
+
+#undef LOCTEXT_NAMESPACE
diff --git a/Engine/Plugins/Experimental/EnhancedInput/Source/EnhancedInput/Public/EnhancedActionKeyMapping.h b/Engine/Plugins/Experimental/EnhancedInput/Source/EnhancedInput/Public/EnhancedActionKeyMapping.h
--- a/Engine/Plugins/Experimental/EnhancedInput/Source/EnhancedInput/Public/EnhancedActionKeyMapping.h
+++ b/Engine/Plugins/Experimental/EnhancedInput/Source/EnhancedInput/Public/EnhancedActionKeyMapping.h
@@ -9,6 +9,33 @@
 
 #include "EnhancedActionKeyMapping.generated.h"
 
+/**
+ * A struct that represents player facing mapping options for an action key mapping.
+ * Use this to set a unique FName for the mapping option to save it, as well as some FText options
+ * for use in UI.
+ */
+USTRUCT(BlueprintType)
+struct FPlayerMappableKeyOptions
+{
+	GENERATED_BODY()
+
+public:
+	
+	FPlayerMappableKeyOptions() = default;
+
+	/** A unique name for this player binding to be saved with. */
+	UPROPERTY(EditAnywhere, BlueprintReadOnly, Category = "Input|PlayerMappable")
+	FName Name = NAME_None;
+	
+	/** The display name that can  */
+	UPROPERTY(EditAnywhere, BlueprintReadOnly, Category = "Input|PlayerMappable")
+	FText DisplayName = FText::GetEmpty();
+
+	/** The category that this player binding is in */
+	UPROPERTY(EditAnywhere, BlueprintReadOnly, Category = "Input|PlayerMappable")
+	FText DisplayCategory = FText::GetEmpty();
+};
+
 /**
  * Defines a mapping between a key activation and the resulting enhanced action
  * An key could be a button press, joystick axis movement, etc.
@@ -30,6 +57,14 @@ struct FEnhancedActionKeyMapping
 
 	//ControllerId ControllerId;	// TODO: Controller id/player id (hybrid?) allowing binding multiple pads to a series of actions.
 
+	/** If true than this ActionKeyMapping will be exposed as a player bindable key */
+	UPROPERTY(EditAnywhere, BlueprintReadOnly, Category = "Input|PlayerMappable")
+	uint8 bIsPlayerMappable : 1;
+
+	/** Options for making this a player mappable keymapping */
+	UPROPERTY(EditAnywhere, BlueprintReadOnly, Category = "Input|PlayerMappable", meta = (editCondition = "bIsPlayerBindable"))
+	FPlayerMappableKeyOptions PlayerMappableOptions {};
+	
 	/**
 	 * If true, then this Key Mapping should be ignored. This is set to true if the key is down
 	 * during a rebuild of it's owning PlayerInput ControlMappings.
@@ -65,6 +100,7 @@ struct FEnhancedActionKeyMapping
 	FEnhancedActionKeyMapping(const UInputAction* InAction = nullptr, const FKey InKey = EKeys::Invalid)
 		: Action(InAction)
 		, Key(InKey)
+		, bIsPlayerMappable(false)
 		, bShouldBeIgnored(false)
 	{}
 
diff --git a/Engine/Plugins/Experimental/EnhancedInput/Source/EnhancedInput/Public/EnhancedInputSubsystemInterface.h b/Engine/Plugins/Experimental/EnhancedInput/Source/EnhancedInput/Public/EnhancedInputSubsystemInterface.h
--- a/Engine/Plugins/Experimental/EnhancedInput/Source/EnhancedInput/Public/EnhancedInputSubsystemInterface.h
+++ b/Engine/Plugins/Experimental/EnhancedInput/Source/EnhancedInput/Public/EnhancedInputSubsystemInterface.h
@@ -6,6 +6,7 @@
 #include "InputActionValue.h"
 #include "InputMappingQuery.h"
 #include "UObject/Interface.h"
+#include "EnhancedActionKeyMapping.h"
 
 #include "EnhancedInputSubsystemInterface.generated.h"
 
@@ -15,6 +16,7 @@ class UInputAction;
 class UEnhancedPlayerInput;
 class UInputModifier;
 class UInputTrigger;
+class UPlayerMappableInputConfig;
 
 // Subsystem interface
 UINTERFACE(MinimalAPI, meta=(CannotImplementInterfaceInBlueprint))
@@ -176,6 +178,19 @@ public:
 	UFUNCTION(BlueprintCallable, BlueprintCosmetic, Category = "Input|Mapping Queries")
 	virtual TArray<FKey> QueryKeysMappedToAction(const UInputAction* Action) const;
 
+	/**
+	 * Replace any currently applied mappings to this key mapping with the given new one.
+	 * Requests a rebuild of the player mappings. 
+	 *
+	 * @return The number of mappings that have been replaced
+	 */
+	UFUNCTION(BlueprintCallable, BlueprintCosmetic, Category = "Input|PlayerMappable", meta=(AutoCreateRefTerm = "Options"))
+	virtual int32 AddPlayerMappedKey(const FName MappingName, const FKey NewKey, const FModifyContextOptions& Options = FModifyContextOptions());
+
+	/** Adds all the input mapping contexts inside of this mappable config. */
+	UFUNCTION(BlueprintCallable, BlueprintCosmetic, Category = "Input|PlayerMappable", meta=(AutoCreateRefTerm = "Options"))
+	virtual void AddPlayerMappableConfig(const UPlayerMappableInputConfig* Config, const FModifyContextOptions& Options = FModifyContextOptions());
+
 private:
 
 	// Forced actions/keys for debug. These will be applied each tick once set even if zeroed, until removed. 
@@ -199,6 +214,9 @@ private:
 	TMap<TWeakObjectPtr<const UInputAction>, FInputActionValue> ForcedActions;
 	TMap<FKey, FInputActionValue> ForcedKeys;
 
+	/** A map of any player mapped keys to the key that they should redirect to instead */
+	TMap<FName, FKey> PlayerMappedSettings;
+
 	EInputMappingRebuildType MappingRebuildPending = EInputMappingRebuildType::None;
 
 	/**
diff --git a/Engine/Plugins/Experimental/EnhancedInput/Source/EnhancedInput/Public/PlayerMappableInputConfig.h b/Engine/Plugins/Experimental/EnhancedInput/Source/EnhancedInput/Public/PlayerMappableInputConfig.h
new file mode 100644
--- /dev/null
+++ b/Engine/Plugins/Experimental/EnhancedInput/Source/EnhancedInput/Public/PlayerMappableInputConfig.h
@@ -0,0 +1,93 @@
+// Copyright Epic Games, Inc. All Rights Reserved.
+
+#pragma once
+
+#include "CoreMinimal.h"
+#include "Engine/DataAsset.h"
+#include "EnhancedActionKeyMapping.h"
+#include "PlayerMappableInputConfig.generated.h"
+
+class UInputMappingContext;
+
+/**
+ * UPlayerMappableInputConfig
+ * 
+ * This represents one set of Player Mappable controller/keymappings. You can use this input config to create
+ * the default mappings for your player to start with in your game. It provides an easy way to get only the player
+ * mappable key actions, and it can be used to add multiple UInputMappingContext's at once to the player.
+ *
+ * Populate this data asset with Input Mapping Contexts that have player bindable actions in them. 
+ */
+UCLASS(BlueprintType, Meta = (DisplayName = "Player Bindable Input Config", ShortTooltip = "Data asset used to define a set of player mappable controller/keyboard mappings."))
+class ENHANCEDINPUT_API UPlayerMappableInputConfig : public UPrimaryDataAsset
+{
+	GENERATED_BODY()
+
+	UPlayerMappableInputConfig(const FObjectInitializer& ObjectInitializer);
+
+public:
+
+#if WITH_EDITOR
+	virtual EDataValidationResult IsDataValid(TArray<FText>& ValidationErrors) override;
+#endif // WITH_EDITOR
+
+	/** Iterate all the player mappable keys in this config in the default input mapping contexts */
+	void ForEachDefaultPlayerMappableKey(TFunctionRef<void(const FEnhancedActionKeyMapping&)> Operation) const;
+
+	/** Get all the player mappable keys in this config. */
+	UFUNCTION(BlueprintCallable, Category = "Input|PlayerMappable")
+	TArray<FEnhancedActionKeyMapping> GetPlayerMappableKeys() const;
+
+	/** Returns the action key mapping for the mapping that matches the given name */
+	UFUNCTION(BlueprintCallable, Category = "Input|PlayerMappable")
+	FEnhancedActionKeyMapping GetMappingByName(const FName MappingName) const;
+
+	/** Returns all the keys mapped to a specific Input Action in this mapping config. */
+	UFUNCTION(BlueprintCallable, Category = "Input|PlayerMappable")
+	TArray<FEnhancedActionKeyMapping> GetKeysBoundToAction(const UInputAction* InAction) const;
+
+	/** Resets this mappable config to use the keys  */
+	UFUNCTION(BlueprintCallable, Category = "Input|PlayerMappable")
+	void ResetToDefault();
+	
+protected:
+	
+	/** The unique name of this config that can be used when saving it */
+	UPROPERTY(EditAnywhere, BlueprintReadOnly, Category = "Input|PlayerMappable")
+	FName ConfigName = NAME_None;
+
+	/** The display name that can be used  */
+	UPROPERTY(EditAnywhere, BlueprintReadOnly, Category = "Input|PlayerMappable")
+	FText ConfigDisplayName = FText::GetEmpty();
+
+	/**  */
+	UPROPERTY(EditAnywhere, BlueprintReadOnly, Category = "Input|PlayerMappable")
+	bool bIsDeprecated = false;
+
+	/** Metadata that can used to store any other related items to your key mapping such as icons, ability assets, etc. */
+	UPROPERTY(EditAnywhere, BlueprintReadOnly, Category = "Input|PlayerMappable")
+	UObject* Metadata = nullptr;
+
+	/** Mapping contexts that make up this Input Config with their associated priority. */
+	UPROPERTY(EditAnywhere, BlueprintReadOnly, Category = "Input|PlayerMappable")
+	TMap<UInputMappingContext*, int32> Contexts;
+
+public:
+
+	/** Return all the Input Mapping contexts that  */
+	UFUNCTION(BlueprintCallable, Category = "Input|PlayerMappable")
+	const TMap<UInputMappingContext*, int32>& GetMappingContexts() const { return Contexts; }
+
+	UFUNCTION(BlueprintCallable, Category = "Input|PlayerMappable")
+	const FName GetConfigName() const { return ConfigName; }
+
+	UFUNCTION(BlueprintCallable, Category = "Input|PlayerMappable")
+	const FText& GetDisplayName() const { return ConfigDisplayName; }
+
+	UFUNCTION(BlueprintCallable, Category = "Input|PlayerMappable")
+	const bool IsDeprecated() const { return bIsDeprecated; }
+
+	/** Get all the player mappable keys in this config. */
+	UFUNCTION(BlueprintCallable, Category = "Input|PlayerMappable")
+	UObject* GetMetadata() const { return Metadata; }
+};
diff --git a/Engine/Plugins/Experimental/EnhancedInput/Source/InputEditor/Private/InputCustomizations.cpp b/Engine/Plugins/Experimental/EnhancedInput/Source/InputEditor/Private/InputCustomizations.cpp
--- a/Engine/Plugins/Experimental/EnhancedInput/Source/InputEditor/Private/InputCustomizations.cpp
+++ b/Engine/Plugins/Experimental/EnhancedInput/Source/InputEditor/Private/InputCustomizations.cpp
@@ -53,10 +53,14 @@ void FEnhancedActionMappingCustomization::CustomizeChildren(TSharedRef<IProperty
 {
 	TSharedPtr<IPropertyHandle> TriggersHandle = PropertyHandle->GetChildHandle(GET_MEMBER_NAME_CHECKED(FEnhancedActionKeyMapping, Triggers));
 	TSharedPtr<IPropertyHandle> ModifiersHandle = PropertyHandle->GetChildHandle(GET_MEMBER_NAME_CHECKED(FEnhancedActionKeyMapping, Modifiers));
+	TSharedPtr<IPropertyHandle> IsPlayerMappableHandle = PropertyHandle->GetChildHandle(GET_MEMBER_NAME_CHECKED(FEnhancedActionKeyMapping, bIsPlayerMappable));
+	TSharedPtr<IPropertyHandle> PlayerBindingOptions = PropertyHandle->GetChildHandle(GET_MEMBER_NAME_CHECKED(FEnhancedActionKeyMapping, PlayerMappableOptions));
 
 	// TODO: ResetToDefault needs to be disabled for arrays
 	ChildBuilder.AddProperty(TriggersHandle.ToSharedRef());
 	ChildBuilder.AddProperty(ModifiersHandle.ToSharedRef());
+	ChildBuilder.AddProperty(IsPlayerMappableHandle.ToSharedRef());
+	ChildBuilder.AddProperty(PlayerBindingOptions.ToSharedRef());
 }
 
 void FEnhancedActionMappingCustomization::RemoveMappingButton_OnClick() const
diff --git a/Engine/Plugins/Experimental/EnhancedInput/Source/InputEditor/Private/InputEditorModule.cpp b/Engine/Plugins/Experimental/EnhancedInput/Source/InputEditor/Private/InputEditorModule.cpp
--- a/Engine/Plugins/Experimental/EnhancedInput/Source/InputEditor/Private/InputEditorModule.cpp
+++ b/Engine/Plugins/Experimental/EnhancedInput/Source/InputEditor/Private/InputEditorModule.cpp
@@ -12,6 +12,7 @@
 #include "HAL/PlatformFileManager.h"
 #include "InputAction.h"
 #include "InputMappingContext.h"
+#include "PlayerMappableInputConfig.h"
 #include "InputCustomizations.h"
 #include "InputModifiers.h"
 #include "IAssetTools.h"
@@ -108,7 +109,9 @@ UObject* UInputMappingContext_Factory::FactoryCreateNew(UClass* Class, UObject*
 }
 
 // InputAction
-UInputAction_Factory::UInputAction_Factory(const class FObjectInitializer& OBJ) : Super(OBJ) {
+UInputAction_Factory::UInputAction_Factory(const class FObjectInitializer& OBJ)
+	: Super(OBJ)
+{
 	SupportedClass = UInputAction::StaticClass();
 	bEditAfterNew = true;
 	bCreateNew = true;
@@ -120,6 +123,21 @@ UObject* UInputAction_Factory::FactoryCreateNew(UClass* Class, UObject* InParent
 	return NewObject<UInputAction>(InParent, Class, Name, Flags | RF_Transactional, Context);
 }
 
+// UPlayerMappableInputConfig_Factory
+UPlayerMappableInputConfig_Factory::UPlayerMappableInputConfig_Factory(const class FObjectInitializer& OBJ)
+	: Super(OBJ)
+{
+	SupportedClass = UPlayerMappableInputConfig::StaticClass();
+	bEditAfterNew = true;
+	bCreateNew = true;
+}
+
+UObject* UPlayerMappableInputConfig_Factory::FactoryCreateNew(UClass* Class, UObject* InParent, FName Name, EObjectFlags Flags, UObject* Context, FFeedbackContext* Warn)
+{
+	check(Class->IsChildOf(UPlayerMappableInputConfig::StaticClass()));
+	return NewObject<UPlayerMappableInputConfig>(InParent, Class, Name, Flags | RF_Transactional, Context);
+}
+
 //
 //// InputTrigger
 //UInputTrigger_Factory::UInputTrigger_Factory(const class FObjectInitializer& OBJ) : Super(OBJ) {
@@ -142,7 +160,8 @@ UObject* UInputAction_Factory::FactoryCreateNew(UClass* Class, UObject* InParent
 // Asset type actions
 // TODO: Move asset type action definitions out?
 
-class FAssetTypeActions_InputContext : public FAssetTypeActions_DataAsset {
+class FAssetTypeActions_InputContext : public FAssetTypeActions_DataAsset
+{
 public:
 	virtual FText GetName() const override { return NSLOCTEXT("AssetTypeActions", "AssetTypeActions_InputMappingContext", "Input Mapping Context"); }
 	virtual uint32 GetCategories() override { return FInputEditorModule::GetInputAssetsCategory(); }
@@ -151,7 +170,8 @@ public:
 	virtual UClass* GetSupportedClass() const override { return UInputMappingContext::StaticClass(); }
 };
 
-class FAssetTypeActions_InputAction : public FAssetTypeActions_DataAsset {
+class FAssetTypeActions_InputAction : public FAssetTypeActions_DataAsset
+{
 public:
 	virtual FText GetName() const override { return NSLOCTEXT("AssetTypeActions", "AssetTypeActions_InputAction", "Input Action"); }
 	virtual uint32 GetCategories() override { return FInputEditorModule::GetInputAssetsCategory(); }
@@ -160,6 +180,15 @@ public:
 	virtual UClass* GetSupportedClass() const override { return UInputAction::StaticClass(); }
 };
 
+class FAssetTypeActions_PlayerMappableInputConfig : public FAssetTypeActions_DataAsset
+{
+public:
+	virtual FText GetName() const override { return NSLOCTEXT("AssetTypeActions", "AssetTypeActions_PlayerBindableInputConfig", "Player Bindable Input Config"); }
+	virtual uint32 GetCategories() override { return FInputEditorModule::GetInputAssetsCategory(); }
+	virtual FColor GetTypeColor() const override { return FColor(127, 255, 255); }
+	virtual FText GetAssetDescription(const FAssetData& AssetData) const override { return NSLOCTEXT("AssetTypeActions", "AssetTypeActions_PlayerBindableInputConfigDesc", "Represents one set of Player Mappable controller/keymappings"); }
+	virtual UClass* GetSupportedClass() const override { return UPlayerMappableInputConfig::StaticClass(); }
+};
 
 
 
@@ -421,6 +450,7 @@ void FInputEditorModule::StartupModule()
 	{
 		RegisterAssetTypeActions(AssetTools, MakeShareable(new FAssetTypeActions_InputAction));
 		RegisterAssetTypeActions(AssetTools, MakeShareable(new FAssetTypeActions_InputContext));
+		RegisterAssetTypeActions(AssetTools, MakeShareable(new FAssetTypeActions_PlayerMappableInputConfig));
 		// TODO: Build these off a button on the InputContext Trigger/Mapping pickers? Would be good to have both.
 		//RegisterAssetTypeActions(AssetTools, MakeShareable(new FAssetTypeActions_InputTrigger));
 		//RegisterAssetTypeActions(AssetTools, MakeShareable(new FAssetTypeActions_InputModifier));
diff --git a/Engine/Plugins/Experimental/EnhancedInput/Source/InputEditor/Public/InputEditorModule.h b/Engine/Plugins/Experimental/EnhancedInput/Source/InputEditor/Public/InputEditorModule.h
--- a/Engine/Plugins/Experimental/EnhancedInput/Source/InputEditor/Public/InputEditorModule.h
+++ b/Engine/Plugins/Experimental/EnhancedInput/Source/InputEditor/Public/InputEditorModule.h
@@ -27,6 +27,14 @@ public:
 	virtual UObject* FactoryCreateNew(UClass* Class, UObject* InParent, FName Name, EObjectFlags Flags, UObject* Context, FFeedbackContext* Warn) override;
 };
 
+UCLASS()
+class INPUTEDITOR_API UPlayerMappableInputConfig_Factory : public UFactory
+{
+	GENERATED_UCLASS_BODY()
+public:
+	virtual UObject* FactoryCreateNew(UClass* Class, UObject* InParent, FName Name, EObjectFlags Flags, UObject* Context, FFeedbackContext* Warn) override;
+};
+
 // TODO: Add trigger/modifier factories and hook up RegisterAssetTypeActions type construction.
 //
 //UCLASS()
