Change the default constructor on player mappable key data so that it gets auto populated with some useful stuff instead of being empty by default.

diff --git a/Engine/Plugins/Experimental/EnhancedInput/Source/EnhancedInput/Public/EnhancedActionKeyMapping.h b/Engine/Plugins/Experimental/EnhancedInput/Source/EnhancedInput/Public/EnhancedActionKeyMapping.h
--- a/Engine/Plugins/Experimental/EnhancedInput/Source/EnhancedInput/Public/EnhancedActionKeyMapping.h
+++ b/Engine/Plugins/Experimental/EnhancedInput/Source/EnhancedInput/Public/EnhancedActionKeyMapping.h
@@ -21,15 +21,28 @@ struct FPlayerMappableKeyOptions
 
 public:
 	
-	FPlayerMappableKeyOptions() = default;
+	FPlayerMappableKeyOptions(const UInputAction* InAction = nullptr)
+	{
+		if(InAction)
+		{
+			const FString& ActionName = InAction->GetName();
+			Name = FName(*ActionName);
+			DisplayName = FText::FromString(ActionName);
+		}
+		else
+		{
+			Name = NAME_None;
+			DisplayName = FText::GetEmpty();
+		}
+	};
 
 	/** A unique name for this player binding to be saved with. */
 	UPROPERTY(EditAnywhere, BlueprintReadOnly, Category = "Input|PlayerMappable")
-	FName Name = NAME_None;
+	FName Name;
 	
 	/** The display name that can  */
 	UPROPERTY(EditAnywhere, BlueprintReadOnly, Category = "Input|PlayerMappable")
-	FText DisplayName = FText::GetEmpty();
+	FText DisplayName;
 
 	/** The category that this player binding is in */
 	UPROPERTY(EditAnywhere, BlueprintReadOnly, Category = "Input|PlayerMappable")
@@ -63,7 +76,7 @@ struct FEnhancedActionKeyMapping
 
 	/** Options for making this a player mappable keymapping */
 	UPROPERTY(EditAnywhere, BlueprintReadOnly, Category = "Input|PlayerMappable", meta = (editCondition = "bIsPlayerBindable"))
-	FPlayerMappableKeyOptions PlayerMappableOptions {};
+	FPlayerMappableKeyOptions PlayerMappableOptions;
 	
 	/**
 	 * If true, then this Key Mapping should be ignored. This is set to true if the key is down
@@ -101,6 +114,7 @@ struct FEnhancedActionKeyMapping
 		: Action(InAction)
 		, Key(InKey)
 		, bIsPlayerMappable(false)
+		, PlayerMappableOptions(InAction)
 		, bShouldBeIgnored(false)
 	{}
 
