Add RemovePlayerMappedKey and RemovePlayerMappableConfig functions to the enhanced input subsystem interface.

#rnx

diff --git a/Engine/Plugins/Experimental/EnhancedInput/Source/EnhancedInput/Private/EnhancedInputSubsystemInterface.cpp b/Engine/Plugins/Experimental/EnhancedInput/Source/EnhancedInput/Private/EnhancedInputSubsystemInterface.cpp
--- a/Engine/Plugins/Experimental/EnhancedInput/Source/EnhancedInput/Private/EnhancedInputSubsystemInterface.cpp
+++ b/Engine/Plugins/Experimental/EnhancedInput/Source/EnhancedInput/Private/EnhancedInputSubsystemInterface.cpp
@@ -391,6 +391,19 @@ int32 IEnhancedInputSubsystemInterface::AddPlayerMappedKey(const FName MappingNa
 	return NumMappingsApplied;
 }
 
+int32 IEnhancedInputSubsystemInterface::RemovePlayerMappedKey(const FName MappingName, const FModifyContextOptions& Options)
+{
+	int32 NumMappingsApplied = 0;
+	if (UEnhancedPlayerInput* const PlayerInput = GetPlayerInput())
+	{
+		NumMappingsApplied = PlayerMappedSettings.Remove(MappingName);
+	}
+
+	RequestRebuildControlMappings(Options);
+	
+	return NumMappingsApplied;
+}
+
 void IEnhancedInputSubsystemInterface::AddPlayerMappableConfig(const UPlayerMappableInputConfig* Config, const FModifyContextOptions& Options)
 {
 	if(Config)
@@ -402,6 +415,17 @@ void IEnhancedInputSubsystemInterface::AddPlayerMappableConfig(const UPlayerMapp
 	}
 }
 
+void IEnhancedInputSubsystemInterface::RemovePlayerMappableConfig(const UPlayerMappableInputConfig* Config, const FModifyContextOptions& Options)
+{
+	if(Config)
+	{
+		for(TPair<UInputMappingContext*, int32> Pair : Config->GetMappingContexts())
+		{
+			RemoveMappingContext(Pair.Key, Options);
+		}	
+	}
+}
+
 template<typename T>
 void DeepCopyPtrArray(const TArray<T*>& From, TArray<T*>& To)
 {
diff --git a/Engine/Plugins/Experimental/EnhancedInput/Source/EnhancedInput/Public/EnhancedInputSubsystemInterface.h b/Engine/Plugins/Experimental/EnhancedInput/Source/EnhancedInput/Public/EnhancedInputSubsystemInterface.h
--- a/Engine/Plugins/Experimental/EnhancedInput/Source/EnhancedInput/Public/EnhancedInputSubsystemInterface.h
+++ b/Engine/Plugins/Experimental/EnhancedInput/Source/EnhancedInput/Public/EnhancedInputSubsystemInterface.h
@@ -187,10 +187,23 @@ public:
 	UFUNCTION(BlueprintCallable, BlueprintCosmetic, Category = "Input|PlayerMappable", meta=(AutoCreateRefTerm = "Options"))
 	virtual int32 AddPlayerMappedKey(const FName MappingName, const FKey NewKey, const FModifyContextOptions& Options = FModifyContextOptions());
 
+	/**
+	 * Remove any player mappings with to the given action
+	 * Requests a rebuild of the player mappings. 
+	 *
+	 * @return The number of mappings that have been removed
+	 */
+	UFUNCTION(BlueprintCallable, BlueprintCosmetic, Category = "Input|PlayerMappable", meta=(AutoCreateRefTerm = "Options"))
+	virtual int32 RemovePlayerMappedKey(const FName MappingName, const FModifyContextOptions& Options = FModifyContextOptions());
+	
 	/** Adds all the input mapping contexts inside of this mappable config. */
 	UFUNCTION(BlueprintCallable, BlueprintCosmetic, Category = "Input|PlayerMappable", meta=(AutoCreateRefTerm = "Options"))
 	virtual void AddPlayerMappableConfig(const UPlayerMappableInputConfig* Config, const FModifyContextOptions& Options = FModifyContextOptions());
 
+	/** Removes all the input mapping contexts inside of this mappable config. */
+	UFUNCTION(BlueprintCallable, BlueprintCosmetic, Category = "Input|PlayerMappable", meta=(AutoCreateRefTerm = "Options"))
+	virtual void RemovePlayerMappableConfig(const UPlayerMappableInputConfig* Config, const FModifyContextOptions& Options = FModifyContextOptions());
+	
 private:
 
 	// Forced actions/keys for debug. These will be applied each tick once set even if zeroed, until removed. 
