Do not add a player mapped key if the mapping name is None. This prevents EI from attempting a rebuild with an empty key name which would map everything to that key.

#rnx

diff --git a/Engine/Plugins/Experimental/EnhancedInput/Source/EnhancedInput/Private/EnhancedInputSubsystemInterface.cpp b/Engine/Plugins/Experimental/EnhancedInput/Source/EnhancedInput/Private/EnhancedInputSubsystemInterface.cpp
--- a/Engine/Plugins/Experimental/EnhancedInput/Source/EnhancedInput/Private/EnhancedInputSubsystemInterface.cpp
+++ b/Engine/Plugins/Experimental/EnhancedInput/Source/EnhancedInput/Private/EnhancedInputSubsystemInterface.cpp
@@ -380,14 +380,20 @@ TArray<FKey> IEnhancedInputSubsystemInterface::QueryKeysMappedToAction(const UIn
 int32 IEnhancedInputSubsystemInterface::AddPlayerMappedKey(const FName MappingName, const FKey NewKey, const FModifyContextOptions& Options)
 {
 	int32 NumMappingsApplied = 0;
-	if (UEnhancedPlayerInput* const PlayerInput = GetPlayerInput())
+	if (MappingName != NAME_None)
 	{
-		PlayerMappedSettings.Add(MappingName, NewKey);
-		++NumMappingsApplied;
-	}
+		if (UEnhancedPlayerInput* const PlayerInput = GetPlayerInput())
+		{
+			PlayerMappedSettings.Add(MappingName, NewKey);
+			++NumMappingsApplied;
+		}
 
-	RequestRebuildControlMappings(Options);
-	
+		RequestRebuildControlMappings(Options);
+	}
+	else
+	{
+		UE_LOG(LogEnhancedInput, Warning, TEXT("Attempted to AddPlayerMappedKey with an invalid MappingName! Mapping has not been applied."));
+	}
 	return NumMappingsApplied;
 }
 
