Remove the final class specifier from the input modifiers. This will make it easier for people to extend them if desired such as making a deadzone driven by game settings.

#rnx

diff --git a/Engine/Plugins/Experimental/EnhancedInput/Source/EnhancedInput/Public/InputModifiers.h b/Engine/Plugins/Experimental/EnhancedInput/Source/EnhancedInput/Public/InputModifiers.h
--- a/Engine/Plugins/Experimental/EnhancedInput/Source/EnhancedInput/Public/InputModifiers.h
+++ b/Engine/Plugins/Experimental/EnhancedInput/Source/EnhancedInput/Public/InputModifiers.h
@@ -89,7 +89,7 @@ enum class EDeadZoneType : uint8
 	*  Values outside this range will be clamped.
 	*/
 UCLASS(NotBlueprintable, MinimalAPI, meta = (DisplayName = "Dead Zone"))
-class UInputModifierDeadZone final : public UInputModifier
+class UInputModifierDeadZone : public UInputModifier
 {
 	GENERATED_BODY()
 
@@ -136,7 +136,7 @@ protected:
 	*  Inverts input per axis
 	*/
 UCLASS(NotBlueprintable, MinimalAPI, meta = (DisplayName = "Negate"))
-class UInputModifierNegate final : public UInputModifier
+class UInputModifierNegate : public UInputModifier
 {
 	GENERATED_BODY()
 
@@ -158,7 +158,7 @@ protected:
 	*  Smooth inputs out over multiple frames
 	*/
 UCLASS(NotBlueprintable, MinimalAPI, meta = (DisplayName = "Smooth"))
-class UInputModifierSmooth final : public UInputModifier
+class UInputModifierSmooth : public UInputModifier
 {
 	GENERATED_BODY()
 
@@ -188,7 +188,7 @@ protected:
 	*  Apply a simple exponential response curve to input values, per axis
 	*/
 UCLASS(NotBlueprintable, MinimalAPI, meta = (DisplayName = "Response Curve - Exponential"))
-class UInputModifierResponseCurveExponential final : public UInputModifier
+class UInputModifierResponseCurveExponential : public UInputModifier
 {
 	GENERATED_BODY()
 
@@ -206,7 +206,7 @@ protected:
 	*  Apply a custom response curve to input values, per axis
 	*/
 UCLASS(NotBlueprintable, MinimalAPI, meta = (DisplayName = "Response Curve - User Defined"))
-class UInputModifierResponseCurveUser final : public	UInputModifier
+class UInputModifierResponseCurveUser : public	UInputModifier
 {
 	GENERATED_BODY()
 
@@ -237,7 +237,7 @@ enum class EFOVScalingType : uint8
 	* Apply FOV dependent scaling to input values, per axis
 	*/
 UCLASS(NotBlueprintable, MinimalAPI, meta = (DisplayName = "FOV Scaling"))
-class UInputModifierFOVScaling final : public	UInputModifier
+class UInputModifierFOVScaling : public	UInputModifier
 {
 	GENERATED_BODY()
 
@@ -293,7 +293,7 @@ enum class EInputAxisSwizzle : uint8
 	* Useful to map a 1D input onto the Y axis of a 2D action.
 	*/
 UCLASS(NotBlueprintable, MinimalAPI, meta = (DisplayName = "Swizzle Input Axis Values"))
-class UInputModifierSwizzleAxis final : public UInputModifier
+class UInputModifierSwizzleAxis : public UInputModifier
 {
 	GENERATED_BODY()
 
@@ -312,7 +312,7 @@ protected:
 	* A user definable group of modifiers that can be easily applied to multiple actions or mappings to save duplication work.
 	*/
 UCLASS(NotBlueprintable, MinimalAPI, meta = (DisplayName = "Modifier Collection"))
-class UInputModifierCollection final : public	UInputModifier
+class UInputModifierCollection : public	UInputModifier
 {
 	GENERATED_BODY()
 
