Enhanced Input: Added InjectInputVectorForAction that can be properly used from Blueprints without an existing FInputActionValue to work with, combining the type from the specified action with a passed in vector

diff --git a/Engine/Plugins/Experimental/EnhancedInput/Source/EnhancedInput/Private/EnhancedInputSubsystemInterface.cpp b/Engine/Plugins/Experimental/EnhancedInput/Source/EnhancedInput/Private/EnhancedInputSubsystemInterface.cpp
--- a/Engine/Plugins/Experimental/EnhancedInput/Source/EnhancedInput/Private/EnhancedInputSubsystemInterface.cpp
+++ b/Engine/Plugins/Experimental/EnhancedInput/Source/EnhancedInput/Private/EnhancedInputSubsystemInterface.cpp
@@ -36,6 +36,12 @@ void IEnhancedInputSubsystemInterface::InjectInputForAction(const UInputAction*
 	}
 }
 
+void IEnhancedInputSubsystemInterface::InjectInputVectorForAction(const UInputAction* Action, FVector Value, const TArray<UInputModifier*>& Modifiers, const TArray<UInputTrigger*>& Triggers)
+{
+	FInputActionValue RawValue((Action != nullptr) ? Action->ValueType : EInputActionValueType::Boolean, Value);
+	InjectInputForAction(Action, RawValue, Modifiers, Triggers);
+}
+
 void IEnhancedInputSubsystemInterface::ClearAllMappings()
 {
 	if (UEnhancedPlayerInput* PlayerInput = GetPlayerInput())
diff --git a/Engine/Plugins/Experimental/EnhancedInput/Source/EnhancedInput/Public/EnhancedInputSubsystemInterface.h b/Engine/Plugins/Experimental/EnhancedInput/Source/EnhancedInput/Public/EnhancedInputSubsystemInterface.h
--- a/Engine/Plugins/Experimental/EnhancedInput/Source/EnhancedInput/Public/EnhancedInputSubsystemInterface.h
+++ b/Engine/Plugins/Experimental/EnhancedInput/Source/EnhancedInput/Public/EnhancedInputSubsystemInterface.h
@@ -79,6 +79,18 @@ public:
 	UFUNCTION(BlueprintCallable, Category="Input", meta=(AutoCreateRefTerm="Modifiers,Triggers"))
 	virtual void InjectInputForAction(const UInputAction* Action, FInputActionValue RawValue, const TArray<UInputModifier*>& Modifiers, const TArray<UInputTrigger*>& Triggers);
 
+	/**
+	 * Input simulation via injection. Runs modifiers and triggers delegates as if the input had come through the underlying input system as FKeys.
+	 * Applies action modifiers and triggers on top.
+	 *
+	 * @param Action		The Input Action to set inject input for
+	 * @param Value			The value to set the action to (the type will be controlled by the Action)
+	 * @param Modifiers		The modifiers to apply to the injected input.
+	 * @param Triggers		The triggers to apply to the injected input.
+	 */
+	UFUNCTION(BlueprintCallable, Category="Input", meta=(AutoCreateRefTerm="Modifiers,Triggers"))
+	virtual void InjectInputVectorForAction(const UInputAction* Action, FVector Value, const TArray<UInputModifier*>& Modifiers, const TArray<UInputTrigger*>& Triggers);
+
 	/**
 	 * Remove all applied mapping contexts.
 	 */
