Move the enhanced input plugin out of Experimental and into the normal /Engine/Plugins folder.

diff --git a/Engine/Plugins/EnhancedInput/Source/EnhancedInput/Private/EnhancedInputSubsystems.cpp b/Engine/Plugins/EnhancedInput/Source/EnhancedInput/Private/EnhancedInputSubsystems.cpp
new file mode 100644
--- /dev/null
+++ b/Engine/Plugins/EnhancedInput/Source/EnhancedInput/Private/EnhancedInputSubsystems.cpp
@@ -0,0 +1,47 @@
+// Copyright Epic Games, Inc. All Rights Reserved.
+
+#include "EnhancedInputSubsystems.h"
+
+#include "Engine/LocalPlayer.h"
+#include "EnhancedInputComponent.h"
+#include "EnhancedInputModule.h"
+#include "EnhancedPlayerInput.h"
+#include "GameFramework/PlayerController.h"
+#include "InputCoreTypes.h"
+#include "InputMappingContext.h"
+#include "InputModifiers.h"
+#include "InputTriggers.h"
+#include "UObject/UObjectIterator.h"
+
+
+// **************************************************************************************************
+// *
+// * UEnhancedInputLocalPlayerSubsystem
+// *
+// **************************************************************************************************
+
+UEnhancedPlayerInput* UEnhancedInputLocalPlayerSubsystem::GetPlayerInput() const
+{	
+	if (APlayerController* PlayerController = GetLocalPlayer()->GetPlayerController(GetWorld()))
+	{
+		return Cast<UEnhancedPlayerInput>(PlayerController->PlayerInput);
+	}
+	return nullptr;
+}
+
+// **************************************************************************************************
+// *
+// * UEnhancedInputEngineSubsystem
+// *
+// **************************************************************************************************
+
+void UEnhancedInputEngineSubsystem::Initialize(FSubsystemCollectionBase& Collection)
+{	
+	PlayerInput = nullptr;
+	//PlayerInput = NewObject<UEnhancedPlayerInput>();	// TODO: Remove UPlayerInput Within? Create an empty player controller?
+}
+
+UEnhancedPlayerInput* UEnhancedInputEngineSubsystem::GetPlayerInput() const
+{
+	return PlayerInput;
+}
diff --git a/Engine/Plugins/EnhancedInput/Source/EnhancedInput/Private/InputMappingContext.cpp b/Engine/Plugins/EnhancedInput/Source/EnhancedInput/Private/InputMappingContext.cpp
new file mode 100644
--- /dev/null
+++ b/Engine/Plugins/EnhancedInput/Source/EnhancedInput/Private/InputMappingContext.cpp
@@ -0,0 +1,40 @@
+// Copyright Epic Games, Inc. All Rights Reserved.
+
+#include "InputMappingContext.h"
+
+#include "EnhancedInputLibrary.h"
+#include "EnhancedInputModule.h"
+
+FEnhancedActionKeyMapping& UInputMappingContext::MapKey(const UInputAction* Action, FKey ToKey)
+{
+	IEnhancedInputModule::Get().GetLibrary()->RequestRebuildControlMappingsUsingContext(this);
+	return Mappings.Add_GetRef(FEnhancedActionKeyMapping(Action, ToKey));
+}
+
+void UInputMappingContext::UnmapKey(const UInputAction* Action, FKey Key)
+{
+	int32 MappingIdx = Mappings.IndexOfByPredicate([&Action, &Key](const FEnhancedActionKeyMapping& Other) { return Other.Action == Action && Other.Key == Key; });
+	if (MappingIdx != INDEX_NONE)
+	{
+		Mappings.RemoveAtSwap(MappingIdx);	// TODO: Preserve order?
+		IEnhancedInputModule::Get().GetLibrary()->RequestRebuildControlMappingsUsingContext(this);
+	}
+}
+
+void UInputMappingContext::UnmapAction(const UInputAction* Action)
+{
+	int32 Found = Mappings.RemoveAllSwap([&Action](const FEnhancedActionKeyMapping& Mapping) { return Mapping.Action == Action; });
+	if (Found > 0)
+	{
+		IEnhancedInputModule::Get().GetLibrary()->RequestRebuildControlMappingsUsingContext(this);
+	}
+}
+
+void UInputMappingContext::UnmapAll()
+{
+	if (Mappings.Num())
+	{
+		Mappings.Empty();
+		IEnhancedInputModule::Get().GetLibrary()->RequestRebuildControlMappingsUsingContext(this);
+	}
+}
\ No newline at end of file
diff --git a/Engine/Plugins/EnhancedInput/Source/EnhancedInput/Public/EnhancedActionKeyMapping.h b/Engine/Plugins/EnhancedInput/Source/EnhancedInput/Public/EnhancedActionKeyMapping.h
new file mode 100644
--- /dev/null
+++ b/Engine/Plugins/EnhancedInput/Source/EnhancedInput/Public/EnhancedActionKeyMapping.h
@@ -0,0 +1,121 @@
+// Copyright Epic Games, Inc. All Rights Reserved.
+
+#pragma once
+
+#include "CoreMinimal.h"
+#include "InputAction.h"
+#include "InputActionValue.h"
+#include "InputCoreTypes.h"
+
+#include "EnhancedActionKeyMapping.generated.h"
+
+/**
+ * A struct that represents player facing mapping options for an action key mapping.
+ * Use this to set a unique FName for the mapping option to save it, as well as some FText options
+ * for use in UI.
+ */
+USTRUCT(BlueprintType)
+struct FPlayerMappableKeyOptions
+{
+	GENERATED_BODY()
+
+public:
+	
+	FPlayerMappableKeyOptions(const UInputAction* InAction = nullptr)
+	{
+		if(InAction)
+		{
+			const FString& ActionName = InAction->GetName();
+			Name = FName(*ActionName);
+			DisplayName = FText::FromString(ActionName);
+		}
+		else
+		{
+			Name = NAME_None;
+			DisplayName = FText::GetEmpty();
+		}
+	};
+
+	/** A unique name for this player binding to be saved with. */
+	UPROPERTY(EditAnywhere, BlueprintReadOnly, Category = "Input|PlayerMappable")
+	FName Name;
+	
+	/** The display name that can  */
+	UPROPERTY(EditAnywhere, BlueprintReadOnly, Category = "Input|PlayerMappable")
+	FText DisplayName;
+
+	/** The category that this player binding is in */
+	UPROPERTY(EditAnywhere, BlueprintReadOnly, Category = "Input|PlayerMappable")
+	FText DisplayCategory = FText::GetEmpty();
+};
+
+/**
+ * Defines a mapping between a key activation and the resulting enhanced action
+ * An key could be a button press, joystick axis movement, etc.
+ * An enhanced action could be MoveForward, Jump, Fire, etc.
+ *
+**/
+USTRUCT(BlueprintType)
+struct FEnhancedActionKeyMapping
+{
+	GENERATED_BODY()
+
+	/** Action to be affected by the key  */
+	UPROPERTY(EditAnywhere, BlueprintReadOnly, Category = "Input")
+	const UInputAction* Action = nullptr;
+
+	/** Key that affect the action. */
+	UPROPERTY(EditAnywhere, BlueprintReadOnly, Category = "Input")
+	FKey Key;
+
+	//ControllerId ControllerId;	// TODO: Controller id/player id (hybrid?) allowing binding multiple pads to a series of actions.
+
+	/** If true than this ActionKeyMapping will be exposed as a player bindable key */
+	UPROPERTY(EditAnywhere, BlueprintReadOnly, Category = "Input|PlayerMappable")
+	uint8 bIsPlayerMappable : 1;
+
+	/** Options for making this a player mappable keymapping */
+	UPROPERTY(EditAnywhere, BlueprintReadOnly, Category = "Input|PlayerMappable", meta = (editCondition = "bIsPlayerBindable"))
+	FPlayerMappableKeyOptions PlayerMappableOptions;
+	
+	/**
+	 * If true, then this Key Mapping should be ignored. This is set to true if the key is down
+	 * during a rebuild of it's owning PlayerInput ControlMappings.
+	 * 
+	 * @see IEnhancedInputSubsystemInterface::RebuildControlMappings
+	 */
+	UPROPERTY(Transient)
+	uint8 bShouldBeIgnored : 1;
+	
+	/**
+	* Trigger qualifiers. If any trigger qualifiers exist the mapping will not trigger unless:
+	* If there are any Explicit triggers in this list at least one of them must be met.
+	* All Implicit triggers in this list must be met.
+	*/
+	UPROPERTY(EditAnywhere, Instanced, BlueprintReadWrite, Category = "Input")
+	TArray<class UInputTrigger*> Triggers;
+
+	/**
+	* Modifiers applied to the raw key value.
+	* These are applied sequentially in array order.
+	*/
+	UPROPERTY(EditAnywhere, Instanced, BlueprintReadWrite, Category = "Input")
+	TArray<class UInputModifier*> Modifiers;
+
+	bool operator==(const FEnhancedActionKeyMapping& Other) const
+	{
+		return (Action == Other.Action &&
+				Key == Other.Key &&
+				Triggers == Other.Triggers &&
+				Modifiers == Other.Modifiers);
+	}
+
+	FEnhancedActionKeyMapping(const UInputAction* InAction = nullptr, const FKey InKey = EKeys::Invalid)
+		: Action(InAction)
+		, Key(InKey)
+		, bIsPlayerMappable(false)
+		, PlayerMappableOptions(InAction)
+		, bShouldBeIgnored(false)
+	{}
+
+};
diff --git a/Engine/Plugins/EnhancedInput/Source/EnhancedInput/Public/EnhancedInputActionDelegateBinding.h b/Engine/Plugins/EnhancedInput/Source/EnhancedInput/Public/EnhancedInputActionDelegateBinding.h
new file mode 100644
--- /dev/null
+++ b/Engine/Plugins/EnhancedInput/Source/EnhancedInput/Public/EnhancedInputActionDelegateBinding.h
@@ -0,0 +1,57 @@
+// Copyright Epic Games, Inc. All Rights Reserved.
+
+#pragma once
+
+#include "CoreMinimal.h"
+#include "UObject/ObjectMacros.h"
+#include "Engine/EngineBaseTypes.h"
+#include "Engine/InputDelegateBinding.h"
+#include "EnhancedPlayerInput.h"
+#include "EnhancedInputActionDelegateBinding.generated.h"
+
+class UInputComponent;
+
+USTRUCT()
+struct ENHANCEDINPUT_API FBlueprintEnhancedInputActionBinding
+{
+	GENERATED_USTRUCT_BODY()
+
+	UPROPERTY()
+	const UInputAction* InputAction = nullptr;
+
+	UPROPERTY()
+	ETriggerEvent TriggerEvent = ETriggerEvent::None;
+
+	UPROPERTY()
+	FName FunctionNameToBind = NAME_None;
+
+	// TODO: bDevelopmentOnly;	// This action delegate will not fire in shipped builds (debug/cheat actions)
+};
+
+
+
+UCLASS()
+class ENHANCEDINPUT_API UEnhancedInputActionDelegateBinding : public UInputDelegateBinding
+{
+	GENERATED_UCLASS_BODY()
+
+	UPROPERTY()
+	TArray<FBlueprintEnhancedInputActionBinding> InputActionDelegateBindings;
+
+	//~ Begin UInputDelegateBinding Interface
+	virtual void BindToInputComponent(UInputComponent* InputComponent, UObject* ObjectToBindTo) const override;
+	//~ End UInputDelegateBinding Interface
+};
+
+UCLASS()
+class ENHANCEDINPUT_API UEnhancedInputActionValueBinding : public UInputDelegateBinding
+{
+	GENERATED_UCLASS_BODY()
+
+	UPROPERTY()
+	TArray<FBlueprintEnhancedInputActionBinding> InputActionValueBindings;
+
+	//~ Begin UInputDelegateBinding Interface
+	virtual void BindToInputComponent(UInputComponent* InputComponent, UObject* ObjectToBindTo) const override;
+	//~ End UInputDelegateBinding Interface
+};
\ No newline at end of file
diff --git a/Engine/Plugins/EnhancedInput/Source/EnhancedInput/Public/EnhancedInputLibrary.h b/Engine/Plugins/EnhancedInput/Source/EnhancedInput/Public/EnhancedInputLibrary.h
new file mode 100644
--- /dev/null
+++ b/Engine/Plugins/EnhancedInput/Source/EnhancedInput/Public/EnhancedInputLibrary.h
@@ -0,0 +1,78 @@
+// Copyright Epic Games, Inc. All Rights Reserved.
+
+#pragma once
+
+#include "CoreMinimal.h"
+#include "EnhancedInputSubsystems.h"
+#include "InputActionValue.h"
+#include "InputMappingQuery.h"
+#include "Kismet/BlueprintFunctionLibrary.h"
+#include "UObject/ObjectMacros.h"
+
+#include "EnhancedInputLibrary.generated.h"
+
+class APlayerController;
+class UInputMappingContext;
+class UInputAction;
+class UEnhancedPlayerInput;
+
+UCLASS()
+class ENHANCEDINPUT_API UEnhancedInputLibrary : public UBlueprintFunctionLibrary
+{
+	GENERATED_BODY()
+
+public:
+	/**
+	 * Call SubsystemPredicate on each registered player and standalone enhanced input subsystem.
+	 */
+	static void ForEachSubsystem(TFunctionRef<void(IEnhancedInputSubsystemInterface*)> SubsystemPredicate);
+
+	/**
+	 * Flag all enhanced input subsystems making use of the mapping context for reapplication of all control mappings at the end of this frame.
+	 * @param Context				Mappings will be rebuilt for all subsystems utilizing this context.
+	 * @param bForceImmediately		The mapping changes will be applied synchronously, rather than at the end of the frame, making them available to the input system on the same frame.
+	 */
+	UFUNCTION(BlueprintCallable, BlueprintCosmetic, Category = "Input")
+	static void RequestRebuildControlMappingsUsingContext(const UInputMappingContext* Context, bool bForceImmediately = false);
+
+
+	/** Breaks an ActionValue into X, Y, Z. Axes not supported by value type will be 0. */
+	UFUNCTION(BlueprintPure, Category = "Input", meta = (NativeBreakFunc))
+	static void BreakInputActionValue(FInputActionValue InActionValue, float& X, float& Y, float& Z);
+
+	/**
+	 * Builds an ActionValue from X, Y, Z. Inherits type from an existing ActionValue. Ignores axis values unused by the provided value type.
+	 * @note Intended for use in Input Modifier Modify Raw overloads to modify an existing Input Action Value.
+	 */
+	UFUNCTION(BlueprintPure, Category = "Input", meta = (Keywords = "construct build", NativeMakeFunc))
+	static FInputActionValue MakeInputActionValue(float X, float Y, float Z, const FInputActionValue& MatchValueType);
+
+
+	// Internal helper functionality
+
+	// GetInputActionvalue internal accessor function for actions that have been bound to from a UEnhancedInputComponent
+	UFUNCTION(BlueprintPure, BlueprintInternalUseOnly, meta = (HidePin = "Action"))
+	static FInputActionValue GetBoundActionValue(AActor* Actor, const UInputAction* Action);
+
+	// FInputActionValue internal auto-converters.
+
+	/** Interpret an InputActionValue as a boolean input */
+	UFUNCTION(BlueprintPure, BlueprintInternalUseOnly, meta = (BlueprintAutocast))
+	static bool Conv_InputActionValueToBool(FInputActionValue InValue);
+
+	/** Interpret an InputActionValue as a 1D axis (float) input */
+	UFUNCTION(BlueprintPure, BlueprintInternalUseOnly, meta = (BlueprintAutocast))
+	static float Conv_InputActionValueToAxis1D(FInputActionValue InValue);
+
+	/** Interpret an InputActionValue as a 2D axis (Vector2D) input */
+	UFUNCTION(BlueprintPure, BlueprintInternalUseOnly, meta = (BlueprintAutocast))
+	static FVector2D Conv_InputActionValueToAxis2D(FInputActionValue InValue);
+
+	/** Interpret an InputActionValue as a 3D axis (Vector) input */
+	UFUNCTION(BlueprintPure, BlueprintInternalUseOnly, meta = (BlueprintAutocast))
+	static FVector Conv_InputActionValueToAxis3D(FInputActionValue ActionValue);
+
+	/** Converts a FInputActionValue to a string */
+	UFUNCTION(BlueprintPure, meta = (DisplayName = "To String (InputActionValue)", CompactNodeTitle = "->", BlueprintAutocast), Category = "Utilities|String")
+	static FString Conv_InputActionValueToString(FInputActionValue ActionValue);
+};
diff --git a/Engine/Plugins/EnhancedInput/Source/EnhancedInput/Public/EnhancedInputSubsystemInterface.h b/Engine/Plugins/EnhancedInput/Source/EnhancedInput/Public/EnhancedInputSubsystemInterface.h
new file mode 100644
--- /dev/null
+++ b/Engine/Plugins/EnhancedInput/Source/EnhancedInput/Public/EnhancedInputSubsystemInterface.h
@@ -0,0 +1,261 @@
+// Copyright Epic Games, Inc. All Rights Reserved.
+
+#pragma once
+
+#include "CoreMinimal.h"
+#include "InputActionValue.h"
+#include "InputMappingQuery.h"
+#include "UObject/Interface.h"
+#include "EnhancedActionKeyMapping.h"
+
+#include "EnhancedInputSubsystemInterface.generated.h"
+
+class APlayerController;
+class UInputMappingContext;
+class UInputAction;
+class UEnhancedPlayerInput;
+class UInputModifier;
+class UInputTrigger;
+class UPlayerMappableInputConfig;
+
+// Subsystem interface
+UINTERFACE(MinimalAPI, meta=(CannotImplementInterfaceInBlueprint))
+class UEnhancedInputSubsystemInterface : public UInterface
+{
+	GENERATED_BODY()
+};
+
+UENUM()
+enum class EInputMappingRebuildType : uint8
+{
+	// No rebuild required.
+	None,
+	// Standard mapping rebuild. Retains existing triggers and modifiers for actions that were previously mapped.
+	Rebuild,
+	// If you have made changes to the triggers/modifiers associated with a UInputAction that was previously mapped a flush is required to reset the tracked data for that action.
+	RebuildWithFlush,
+};
+
+/** Passed in as params for Adding/Remove input contexts */
+USTRUCT(BlueprintType)
+struct FModifyContextOptions
+{
+	GENERATED_BODY()
+	
+	FModifyContextOptions()
+		: bIgnoreAllPressedKeysUntilRelease(true)
+		, bForceImmediately(false)
+	{}
+
+	// If true than any keys that are pressed during the rebuild of control mappings will be ignored until they are released.
+	UPROPERTY(EditAnywhere, BlueprintReadWrite, Category = "Input")
+	uint8 bIgnoreAllPressedKeysUntilRelease : 1;
+
+	// The mapping changes will be applied synchronously, rather than at the end of the frame, making them available to the input system on the same frame.
+	UPROPERTY(EditAnywhere, BlueprintReadWrite, Category = "Input")
+	uint8 bForceImmediately : 1;
+};
+
+// Includes native functionality shared between all subsystems
+class ENHANCEDINPUT_API IEnhancedInputSubsystemInterface
+{
+	friend class FEnhancedInputModule;
+
+	GENERATED_BODY()
+
+public:
+
+	virtual UEnhancedPlayerInput* GetPlayerInput() const = 0;
+
+	/**
+	 * Input simulation via injection. Runs modifiers and triggers delegates as if the input had come through the underlying input system as FKeys.
+	 * Applies action modifiers and triggers on top.
+	 *
+	 * @param Action		The Input Action to set inject input for
+	 * @param RawValue		The value to set the action to
+	 * @param Modifiers		The modifiers to apply to the injected input.
+	 * @param Triggers		The triggers to apply to the injected input.
+	 */
+	UFUNCTION(BlueprintCallable, Category="Input", meta=(AutoCreateRefTerm="Modifiers,Triggers"))
+	virtual void InjectInputForAction(const UInputAction* Action, FInputActionValue RawValue, const TArray<UInputModifier*>& Modifiers, const TArray<UInputTrigger*>& Triggers);
+
+	/**
+	 * Input simulation via injection. Runs modifiers and triggers delegates as if the input had come through the underlying input system as FKeys.
+	 * Applies action modifiers and triggers on top.
+	 *
+	 * @param Action		The Input Action to set inject input for
+	 * @param Value			The value to set the action to (the type will be controlled by the Action)
+	 * @param Modifiers		The modifiers to apply to the injected input.
+	 * @param Triggers		The triggers to apply to the injected input.
+	 */
+	UFUNCTION(BlueprintCallable, Category="Input", meta=(AutoCreateRefTerm="Modifiers,Triggers"))
+	virtual void InjectInputVectorForAction(const UInputAction* Action, FVector Value, const TArray<UInputModifier*>& Modifiers, const TArray<UInputTrigger*>& Triggers);
+
+	/**
+	 * Remove all applied mapping contexts.
+	 */
+	UFUNCTION(BlueprintCallable, BlueprintCosmetic, Category = "Input")
+	virtual void ClearAllMappings();
+
+	/**
+	 * Add a control mapping context.
+	 * 
+	 * @param MappingContext		A set of key to action mappings to apply to this player
+	 * @param Priority				Higher priority mappings will be applied first and, if they consume input, will block lower priority mappings.
+	 * @param bIgnoreAllPressedKeysUntilRelease	If true than any keys that are pressed during the rebuild of control mappings will be ignored until they are released.
+	 */
+	UE_DEPRECATED(5.0, "This version of AddMappingContext has been deprecated, please use the version that takes in a FModifyContextOptions instead.")
+	virtual void AddMappingContext(const UInputMappingContext* MappingContext, int32 Priority, const bool bIgnoreAllPressedKeysUntilRelease);
+	
+	/**
+	 * Add a control mapping context.
+	 * @param MappingContext		A set of key to action mappings to apply to this player
+	 * @param Priority				Higher priority mappings will be applied first and, if they consume input, will block lower priority mappings.
+	 * @param Options				Options to consider when adding this mapping context.
+	 */
+	UFUNCTION(BlueprintCallable, BlueprintCosmetic, Category = "Input", meta=(AutoCreateRefTerm = "Options"))
+	virtual void AddMappingContext(const UInputMappingContext* MappingContext, int32 Priority, const FModifyContextOptions& Options = FModifyContextOptions());
+	
+	/**
+	 * Remove a specific control context. 
+	 * This is safe to call even if the context is not applied.
+	 * @param MappingContext		Context to remove from the player
+	 * @param bIgnoreAllPressedKeysUntilRelease	If true than any keys that are pressed during the rebuild of control mappings will be ignored until they are released.
+	 */
+	UE_DEPRECATED(5.0, "This version of RemoveMappingContext has been deprecated, please use the version that takes in a FModifyContextOptions instead.")
+	virtual void RemoveMappingContext(const UInputMappingContext* MappingContext, const bool bIgnoreAllPressedKeysUntilRelease);
+
+	/**
+	* Remove a specific control context. 
+	* This is safe to call even if the context is not applied.
+	* @param MappingContext		Context to remove from the player
+	* @param Options			Options to consider when removing this input mapping context
+	*/
+	UFUNCTION(BlueprintCallable, BlueprintCosmetic, Category = "Input", meta=(AutoCreateRefTerm = "Options"))
+	virtual void RemoveMappingContext(const UInputMappingContext* MappingContext, const FModifyContextOptions& Options = FModifyContextOptions());
+	
+	/**
+	 * Flag player for reapplication of all mapping contexts at the end of this frame.
+	 * This is called automatically when adding or removing mappings contexts.
+	 * @param bForceImmediately		THe mapping changes will be applied synchronously, rather than at the end of the frame, making them available to the input system on the same frame.
+	 */
+	UE_DEPRECATED(5.0, "This version of RequestRebuildControlMappings has been deprecated, please use the version that takes in a FModifyContextOptions instead.")
+	virtual void RequestRebuildControlMappings(bool bForceImmediately, const bool bIgnoreAllPressedKeysUntilRelease = true);
+
+	/**
+	* Flag player for reapplication of all mapping contexts at the end of this frame.
+	* This is called automatically when adding or removing mappings contexts.
+	*
+	* @param Options		Options to consider when removing this input mapping context
+	*/
+	UFUNCTION(BlueprintCallable, BlueprintCosmetic, Category = "Input", meta=(AutoCreateRefTerm = "Options"))
+	virtual void RequestRebuildControlMappings(const FModifyContextOptions& Options = FModifyContextOptions(), EInputMappingRebuildType RebuildType = EInputMappingRebuildType::Rebuild);
+
+	/**
+	 * Check if a key mapping is safe to add to a given mapping context within the set of active contexts currently applied to the player controller.
+	 * @param InputContext		Mapping context to which the action/key mapping is intended to be added
+	 * @param Action			Action that can be triggered by the key
+	 * @param Key				Key that will provide input values towards triggering the action
+	 * @param OutIssues			Issues that may cause this mapping to be invalid (at your discretion). Any potential issues will be recorded, even if not present in FatalIssues.
+	 * @param BlockingIssues	All issues that should be considered fatal as a bitset.
+	 * @return					Summary of resulting issues.
+	 * @see QueryMapKeyInContextSet
+	 */
+	UFUNCTION(BlueprintCallable, BlueprintCosmetic, Category = "Input|Mapping Queries")
+	virtual EMappingQueryResult QueryMapKeyInActiveContextSet(const UInputMappingContext* InputContext, const UInputAction* Action, FKey Key, TArray<FMappingQueryIssue>& OutIssues, EMappingQueryIssue BlockingIssues/* = DefaultMappingIssues::StandardFatal*/);
+
+	/**
+	 * Check if a key mapping is safe to add to a collection of mapping contexts
+	 * @param PrioritizedActiveContexts	Set of mapping contexts to test against ordered by priority such that earlier entries take precedence over later ones.
+	 * @param InputContext		Mapping context to which the action/key mapping is intended to be applied. NOTE: This context must be present in PrioritizedActiveContexts.
+	 * @param Action			Action that is triggered by the key
+	 * @param Key				Key that will provide input values towards triggering the action
+	 * @param OutIssues			Issues that may cause this mapping to be invalid (at your discretion). Any potential issues will be recorded, even if not present in FatalIssues.
+	 * @param BlockingIssues	All issues that should be considered fatal as a bitset.
+	 * @return					Summary of resulting issues.
+	 * @see QueryMapKeyInActiveContextSet
+	 */
+	UFUNCTION(BlueprintCallable, BlueprintCosmetic, Category = "Input|Mapping Queries")
+	virtual EMappingQueryResult QueryMapKeyInContextSet(const TArray<UInputMappingContext*>& PrioritizedActiveContexts, const UInputMappingContext* InputContext, const UInputAction* Action, FKey Key, TArray<FMappingQueryIssue>& OutIssues, EMappingQueryIssue BlockingIssues/* = DefaultMappingIssues::StandardFatal*/);
+
+	/**
+	 * Check if a mapping context is applied to this subsystem's owner.
+	 */
+	UFUNCTION(BlueprintCallable, BlueprintCosmetic, Category = "Input|Mapping Queries")	// TODO: BlueprintPure would be nicer. Move into library?
+	virtual bool HasMappingContext(const UInputMappingContext* MappingContext) const;
+
+	/**
+	 * Returns the keys mapped to the given action in the active input mapping contexts.
+	 */
+	UFUNCTION(BlueprintCallable, BlueprintCosmetic, Category = "Input|Mapping Queries")
+	virtual TArray<FKey> QueryKeysMappedToAction(const UInputAction* Action) const;
+
+	/**
+	 * Replace any currently applied mappings to this key mapping with the given new one.
+	 * Requests a rebuild of the player mappings. 
+	 *
+	 * @return The number of mappings that have been replaced
+	 */
+	UFUNCTION(BlueprintCallable, BlueprintCosmetic, Category = "Input|PlayerMappable", meta=(AutoCreateRefTerm = "Options"))
+	virtual int32 AddPlayerMappedKey(const FName MappingName, const FKey NewKey, const FModifyContextOptions& Options = FModifyContextOptions());
+
+	/**
+	 * Remove any player mappings with to the given action
+	 * Requests a rebuild of the player mappings. 
+	 *
+	 * @return The number of mappings that have been removed
+	 */
+	UFUNCTION(BlueprintCallable, BlueprintCosmetic, Category = "Input|PlayerMappable", meta=(AutoCreateRefTerm = "Options"))
+	virtual int32 RemovePlayerMappedKey(const FName MappingName, const FModifyContextOptions& Options = FModifyContextOptions());
+	
+	/** Adds all the input mapping contexts inside of this mappable config. */
+	UFUNCTION(BlueprintCallable, BlueprintCosmetic, Category = "Input|PlayerMappable", meta=(AutoCreateRefTerm = "Options"))
+	virtual void AddPlayerMappableConfig(const UPlayerMappableInputConfig* Config, const FModifyContextOptions& Options = FModifyContextOptions());
+
+	/** Removes all the input mapping contexts inside of this mappable config. */
+	UFUNCTION(BlueprintCallable, BlueprintCosmetic, Category = "Input|PlayerMappable", meta=(AutoCreateRefTerm = "Options"))
+	virtual void RemovePlayerMappableConfig(const UPlayerMappableInputConfig* Config, const FModifyContextOptions& Options = FModifyContextOptions());
+	
+private:
+
+	// Forced actions/keys for debug. These will be applied each tick once set even if zeroed, until removed. 
+	void ApplyForcedInput(const UInputAction* Action, FInputActionValue Value);
+	void ApplyForcedInput(FKey Key, FInputActionValue Value);
+	void RemoveForcedInput(const UInputAction* Action);
+	void RemoveForcedInput(FKey Key);
+	void TickForcedInput(float DeltaTime);
+
+	void InjectChordBlockers(const TArray<int32>& ChordedMappings);
+	bool HasTriggerWith(TFunctionRef<bool(const class UInputTrigger*)> TestFn, const TArray<class UInputTrigger*>& Triggers);
+
+	/**
+	 * Reapply all control mappings to players pending a rebuild
+	 */
+	void RebuildControlMappings();
+
+	/** Convert input settings axis config to modifiers for a given mapping */
+	void ApplyAxisPropertyModifiers(UEnhancedPlayerInput* PlayerInput, struct FEnhancedActionKeyMapping& Mapping) const;
+
+	TMap<TWeakObjectPtr<const UInputAction>, FInputActionValue> ForcedActions;
+	TMap<FKey, FInputActionValue> ForcedKeys;
+
+	/** A map of any player mapped keys to the key that they should redirect to instead */
+	TMap<FName, FKey> PlayerMappedSettings;
+
+	EInputMappingRebuildType MappingRebuildPending = EInputMappingRebuildType::None;
+
+	/**
+	 * A flag that will be set when adding/removing a mapping context.
+	 *
+	 * If this is true, then any keys that are pressed when control mappings are rebuilt will be ignored 
+	 * by the new Input context after being until the key is lifted
+	 */
+	bool bIgnoreAllPressedKeysUntilReleaseOnRebuild = true;
+
+	bool bMappingRebuildPending = false;
+
+	// Debug visualization implemented in EnhancedInputSubsystemsDebug.cpp
+	void ShowDebugInfo(class UCanvas* Canvas);
+	void ShowDebugActionModifiers(UCanvas* Canvas, const UInputAction* Action);
+	static void PurgeDebugVisualizations();
+};
\ No newline at end of file
diff --git a/Engine/Plugins/EnhancedInput/Source/EnhancedInput/Public/EnhancedInputSubsystems.h b/Engine/Plugins/EnhancedInput/Source/EnhancedInput/Public/EnhancedInputSubsystems.h
new file mode 100644
--- /dev/null
+++ b/Engine/Plugins/EnhancedInput/Source/EnhancedInput/Public/EnhancedInputSubsystems.h
@@ -0,0 +1,46 @@
+// Copyright Epic Games, Inc. All Rights Reserved.
+
+#pragma once
+
+#include "CoreMinimal.h"
+#include "EnhancedInputSubsystemInterface.h"
+#include "Subsystems/LocalPlayerSubsystem.h"
+#include "Subsystems/EngineSubsystem.h"
+
+#include "EnhancedInputSubsystems.generated.h"
+
+// Per local player input subsystem
+UCLASS()
+class ENHANCEDINPUT_API UEnhancedInputLocalPlayerSubsystem : public ULocalPlayerSubsystem, public IEnhancedInputSubsystemInterface
+{
+	GENERATED_BODY()
+
+public:
+
+	// Begin IEnhancedInputSubsystemInterface
+	virtual UEnhancedPlayerInput* GetPlayerInput() const override;
+	// End IEnhancedInputSubsystemInterface
+};
+
+// Global input handling subsystem
+// TODO: For now this is a non-functional placeholder that shuold be enabled when the Enhanced Input Plugin is enabled by default.
+UCLASS(NotBlueprintType)
+class ENHANCEDINPUT_API UEnhancedInputEngineSubsystem : public UEngineSubsystem, public IEnhancedInputSubsystemInterface
+{
+	GENERATED_BODY()
+
+private:
+	UPROPERTY()
+	UEnhancedPlayerInput* PlayerInput;
+
+public:
+
+	// Begin USubsystem
+	virtual void Initialize(FSubsystemCollectionBase& Collection) override;
+	// End USubsystem
+
+	// Begin IEnhancedInputSubsystemInterface
+	virtual UEnhancedPlayerInput* GetPlayerInput() const override;
+	// End IEnhancedInputSubsystemInterface
+
+};
diff --git a/Engine/Plugins/EnhancedInput/Source/EnhancedInput/Public/EnhancedPlayerInput.h b/Engine/Plugins/EnhancedInput/Source/EnhancedInput/Public/EnhancedPlayerInput.h
new file mode 100644
--- /dev/null
+++ b/Engine/Plugins/EnhancedInput/Source/EnhancedInput/Public/EnhancedPlayerInput.h
@@ -0,0 +1,125 @@
+// Copyright Epic Games, Inc. All Rights Reserved.
+
+#pragma once
+
+#include "CoreMinimal.h"
+#include "EnhancedActionKeyMapping.h"
+#include "GameFramework/PlayerInput.h"
+#include "InputAction.h"
+#include "InputActionValue.h"
+#include "InputTriggers.h"
+
+#include "EnhancedPlayerInput.generated.h"
+
+// Internal representation containing event variants
+enum class ETriggerEventInternal : uint8;
+enum class EKeyEvent : uint8;
+
+/**
+* UEnhancedPlayerInput : UPlayerInput extensions for enhanced player input system
+*/
+UCLASS(Within = PlayerController, config = Input, transient)
+class ENHANCEDINPUT_API UEnhancedPlayerInput : public UPlayerInput
+{
+	friend class IEnhancedInputSubsystemInterface;
+	friend class UEnhancedInputLibrary;
+	friend struct FInputTestHelper;
+
+	GENERATED_BODY()
+
+public:
+
+	// TODO: Can we avoid exposing this?
+	const FInputActionInstance* FindActionInstanceData(const UInputAction* ForAction) const { return ActionInstanceData.Find(ForAction); }
+
+	/** Retrieve the current value of an action for this player.
+	* Note: If the action is not currently triggering this will return a zero value of the appropriate value type, ignoring any ongoing inputs.
+	*/
+	FInputActionValue GetActionValue(const UInputAction* ForAction) const;
+
+	// Input simulation via injection. Runs modifiers and triggers delegates as if the input had come through the underlying input system as FKeys. Applies action modifiers and triggers on top.
+	void InjectInputForAction(const UInputAction* Action, FInputActionValue RawValue, const TArray<UInputModifier*>& Modifiers = {}, const TArray<UInputTrigger*>& Triggers = {});
+
+	virtual bool InputKey(const FInputKeyParams& Params) override;
+	
+protected:
+
+	// Applies modifiers and triggers without affecting keys read by the base input system
+	virtual void ProcessInputStack(const TArray<UInputComponent*>& InputComponentStack, const float DeltaTime, const bool bGamePaused) override;
+
+	// Causes key to be consumed if it is affecting an action.
+	virtual bool IsKeyHandledByAction(FKey Key) const override;
+
+private:
+
+	/** Add a player specific action mapping.
+	* Returns index into EnhancedActionMappings array.
+	*/
+	int32 AddMapping(const FEnhancedActionKeyMapping& Mapping);
+	void ClearAllMappings();
+
+	virtual void ConditionalBuildKeyMappings_Internal() const override;
+
+	// Perform a first pass run of modifiers on an action instance
+	void InitializeMappingActionModifiers(const FEnhancedActionKeyMapping& Mapping);
+
+	FInputActionValue ApplyModifiers(const TArray<UInputModifier*>& Modifiers, FInputActionValue RawValue, float DeltaTime) const;						// Pre-modified (raw) value
+	ETriggerEventInternal GetTriggerStateChangeEvent(ETriggerState LastTriggerState, ETriggerState NewTriggerState) const;
+	ETriggerEvent ConvertInternalTriggerEvent(ETriggerEventInternal Event) const;	// Collapse a detailed internal trigger event into a friendly representation
+	void ProcessActionMappingEvent(const UInputAction* Action, float DeltaTime, bool bGamePaused, FInputActionValue RawValue, EKeyEvent KeyEvent, const TArray<UInputModifier*>& Modifiers, const TArray<UInputTrigger*>& Triggers);
+
+	FInputActionInstance& FindOrAddActionEventData(const UInputAction* Action) const;
+
+	template<typename T>
+	void GatherActionEventDataForActionMap(const T& ActionMap, TMap<const UInputAction*, FInputActionInstance>& FoundActionEventData) const;
+
+	/** Currently applied key mappings
+	 * Note: Source reference only. Use EnhancedActionMappings for the actual mappings (with properly instanced triggers/modifiers)
+	 */
+	UPROPERTY(Transient)
+	TMap<const class UInputMappingContext*, int32> AppliedInputContexts;
+
+	/** This player's version of the Action Mappings */
+	UPROPERTY(Transient)
+	TArray<FEnhancedActionKeyMapping> EnhancedActionMappings;
+
+	// Number of active binds by key
+	TMap<FKey, int32> EnhancedKeyBinds;
+
+	/** Tracked action values. Queryable. */
+	UPROPERTY(Transient)
+	mutable TMap<const UInputAction*, FInputActionInstance> ActionInstanceData;
+
+	/** Actions which had actuated events at the last call to ProcessInputStack (held/pressed/released) */
+	TSet<const UInputAction*> ActionsWithEventsThisTick;
+
+	/**
+	 * A map of Keys to the amount they were depressed this frame. This is reset with each call to ProcessInputStack
+	 * and is populated within UEnhancedPlayerInput::InputKey.
+	 */
+	UPROPERTY(Transient)
+	TMap<FKey, FVector> KeysPressedThisTick;
+
+	struct FInjectedInput
+	{
+		FInputActionValue RawValue;
+		TArray<UInputTrigger*> Triggers;
+		TArray<UInputModifier*> Modifiers;
+	};
+	struct FInjectedInputArray
+	{
+		TArray<FInjectedInput> Injected;
+	};
+
+	/** Inputs injected since the last call to ProcessInputStack */
+	TMap<const UInputAction*, FInjectedInputArray> InputsInjectedThisTick;
+
+	/** Last frame's injected inputs */
+	TSet<const UInputAction*> LastInjectedActions;
+
+	/** The last time of the last frame that was processed in ProcessPlayerInput */
+	float LastFrameTime = 0.0f;
+
+	/** Delta seconds between frames calculated with UWorld::GetRealTimeSeconds */
+	float RealTimeDeltaSeconds = 0.0f;
+};
diff --git a/Engine/Plugins/EnhancedInput/Source/EnhancedInput/Public/InputDebugKeyDelegateBinding.h b/Engine/Plugins/EnhancedInput/Source/EnhancedInput/Public/InputDebugKeyDelegateBinding.h
new file mode 100644
--- /dev/null
+++ b/Engine/Plugins/EnhancedInput/Source/EnhancedInput/Public/InputDebugKeyDelegateBinding.h
@@ -0,0 +1,46 @@
+// Copyright Epic Games, Inc. All Rights Reserved.
+
+#pragma once
+
+#include "CoreMinimal.h"
+#include "UObject/ObjectMacros.h"
+#include "Engine/EngineBaseTypes.h"
+#include "Engine/InputDelegateBinding.h"
+#include "Framework/Commands/InputChord.h"
+#include "InputDebugKeyDelegateBinding.generated.h"
+
+class UInputComponent;
+
+USTRUCT()
+struct ENHANCEDINPUT_API FBlueprintInputDebugKeyDelegateBinding
+{
+	GENERATED_BODY()
+
+	UPROPERTY()
+	FInputChord InputChord;
+
+	UPROPERTY()
+	TEnumAsByte<EInputEvent> InputKeyEvent = IE_Pressed;
+
+	UPROPERTY()
+	FName FunctionNameToBind = NAME_None;
+
+	UPROPERTY()
+	bool bExecuteWhenPaused = false;
+
+	// TODO: bConsumeInput?
+};
+
+UCLASS()
+class ENHANCEDINPUT_API UInputDebugKeyDelegateBinding : public UInputDelegateBinding
+{
+	GENERATED_BODY()
+
+public:
+	UPROPERTY()
+	TArray<FBlueprintInputDebugKeyDelegateBinding> InputDebugKeyDelegateBindings;
+
+	//~ Begin UInputDelegateBinding Interface
+	virtual void BindToInputComponent(UInputComponent* InputComponent, UObject* ObjectToBindTo) const override;
+	//~ End UInputDelegateBinding Interface
+};
diff --git a/Engine/Plugins/EnhancedInput/Source/EnhancedInput/Public/InputMappingContext.h b/Engine/Plugins/EnhancedInput/Source/EnhancedInput/Public/InputMappingContext.h
new file mode 100644
--- /dev/null
+++ b/Engine/Plugins/EnhancedInput/Source/EnhancedInput/Public/InputMappingContext.h
@@ -0,0 +1,77 @@
+// Copyright Epic Games, Inc. All Rights Reserved.
+
+#pragma once
+
+#include "CoreMinimal.h"
+#include "EnhancedActionKeyMapping.h"
+#include "Engine/DataAsset.h"
+
+#include "InputMappingContext.generated.h"
+
+class UInputAction;
+
+/**
+* UInputMappingContext : A collection of key to action mappings for a specific input context
+* Could be used to:
+*	Store predefined controller mappings (allow switching between controller config variants). TODO: Build a system allowing redirects of UInputMappingContexts to handle this.
+*	Define per-vehicle control mappings
+*	Define context specific mappings (e.g. I switch from a gun (shoot action) to a grappling hook (reel in, reel out, disconnect actions).
+*	Define overlay mappings to be applied on top of existing control mappings (e.g. Hero specific action mappings in a MOBA)
+*/
+UCLASS(BlueprintType, config = Input)
+class ENHANCEDINPUT_API UInputMappingContext : public UDataAsset
+{
+	GENERATED_BODY()
+
+public:
+
+	// Localized context descriptor
+	UPROPERTY(BlueprintReadOnly, EditAnywhere, Category = "Config", DisplayName = "Description")
+	FText ContextDescription;
+
+protected:
+	// List of key to action mappings.
+	UPROPERTY(config, BlueprintReadOnly, EditAnywhere, Category = "Mappings")
+	TArray<FEnhancedActionKeyMapping> Mappings;
+
+	friend class FInputContextDetails;
+
+public:
+
+	/**
+	* Mapping accessors.
+	* Note: Use UEnhancedInputLibrary::RequestRebuildControlMappingsForContext to invoke changes made to an FEnhancedActionKeyMapping
+	*/
+	const TArray<FEnhancedActionKeyMapping>& GetMappings() const { return Mappings; }
+	FEnhancedActionKeyMapping& GetMapping(TArray<FEnhancedActionKeyMapping>::SizeType Index) { return Mappings[Index]; }
+
+	// TODO: Don't want to encourage Map/Unmap calls here where context switches would be desirable. These are intended for use in the config/binding screen only.
+
+	/**
+	* Map a key to an action within the mapping context.
+	*/
+	UFUNCTION(BlueprintCallable, Category = "Mapping")
+	FEnhancedActionKeyMapping& MapKey(const UInputAction* Action, FKey ToKey);
+
+	/**
+	* Unmap a key from an action within the mapping context.
+	*/
+	UFUNCTION(BlueprintCallable, Category = "Mapping")
+	void UnmapKey(const UInputAction* Action, FKey Key);
+
+	/**
+	* Unmap all key maps to an action within the mapping context.
+	*/
+	UFUNCTION(BlueprintCallable, Category = "Mapping")
+	void UnmapAction(const UInputAction* Action);
+
+	/**
+	* Unmap everything within the mapping context.
+	*/
+	UFUNCTION(BlueprintCallable, Category = "Mapping")
+	void UnmapAll();
+};
+
+// ************************************************************************************************
+// ************************************************************************************************
+// ************************************************************************************************
diff --git a/Engine/Plugins/EnhancedInput/Source/EnhancedInput/Public/InputModifiers.h b/Engine/Plugins/EnhancedInput/Source/EnhancedInput/Public/InputModifiers.h
new file mode 100644
--- /dev/null
+++ b/Engine/Plugins/EnhancedInput/Source/EnhancedInput/Public/InputModifiers.h
@@ -0,0 +1,334 @@
+// Copyright Epic Games, Inc. All Rights Reserved.
+
+#pragma once
+
+#include "CoreMinimal.h"
+#include "InputActionValue.h"
+#include "InputModifiers.generated.h"
+
+class UEnhancedPlayerInput;
+
+// NOTE: Deprecated. Do not use.
+UENUM(BlueprintType)
+enum class EModifierExecutionPhase : uint8
+{
+	// Deprecated. Do not use.
+	PerInput,
+	// Deprecated. Do not use.
+	FinalValue,
+
+	NumPhases			UMETA(Hidden)
+};
+
+/**
+Base class for building modifiers.
+*/
+UCLASS(Abstract, Blueprintable, EditInlineNew, CollapseCategories, Config = Input, defaultconfig, configdonotcheckdefaults)
+class ENHANCEDINPUT_API UInputModifier : public UObject
+{
+	GENERATED_BODY()
+
+protected:
+
+	/** ModifyRaw implementation. Override this to alter input values in native code.
+	 * @param CurrentValue - The modified value returned by the previous modifier in the chain, or the base input device value if this is the first modifier in the chain.
+	 * @param DeltaTime - Elapsed time since last input tick.
+	 * @return Modified value. Note that whilst the returned value can be of any FInputActionValueType it will be reset to the value type of the associated action before any further processing.
+	 */
+	virtual FInputActionValue ModifyRaw_Implementation(const UEnhancedPlayerInput* PlayerInput, FInputActionValue CurrentValue, float DeltaTime)
+	{
+		return CurrentValue;
+	}
+
+	UE_DEPRECATED(4.26, "Execution phase is deprecated.")
+	virtual EModifierExecutionPhase GetExecutionPhase_Implementation() const { return EModifierExecutionPhase::PerInput; }
+
+	virtual FLinearColor GetVisualizationColor_Implementation(FInputActionValue SampleValue, FInputActionValue FinalValue) const
+	{
+		const float Intensity = FMath::Min(1.f, FinalValue.GetMagnitude());	// TODO: 3D visualization!
+		return FLinearColor(Intensity, Intensity, Intensity);
+	}
+
+public:
+
+	/**
+	 * ModifyRaw
+	 * Will be called by each modifier in the modifier chain
+	 * @param CurrentValue - The modified value returned by the previous modifier in the chain, or the base raw value if this is the first modifier in the chain.
+	 */
+	UFUNCTION(BlueprintCallable, BlueprintNativeEvent, Category = "Modifier")
+	FInputActionValue ModifyRaw(const UEnhancedPlayerInput* PlayerInput, FInputActionValue CurrentValue, float DeltaTime) const;
+
+	/**
+	 * GetExecutionPhase - Note: Deprecated
+	 */
+	UE_DEPRECATED(4.26, "Execution phase is deprecated. This call can be safely removed.")
+	UFUNCTION(BlueprintNativeEvent, Category = "Modifier", meta=(DeprecatedFunction, DeprecationMessage="Execution phase is deprecated."))
+	EModifierExecutionPhase GetExecutionPhase() const;
+
+	/**
+	 * Helper to allow debug visualization of the modifier.
+	 * @param SampleValue - The base input action value pre-modification (ranging -1 -> 1 across all applicable axes).
+	 * @param FinalValue - The post-modification input action value for the provided SampleValue.
+	 */
+	UFUNCTION(BlueprintNativeEvent, Category = "Modifier")
+	FLinearColor GetVisualizationColor(FInputActionValue SampleValue, FInputActionValue FinalValue) const;
+};
+
+UENUM()
+enum class EDeadZoneType : uint8
+{
+	// Apply dead zone to axes individually. This will result in input being chamfered at the corners for 2d/3d axis inputs, and matches the original UE4 deadzone logic.
+	Axial,
+	// Apply dead zone logic to all axes simultaneously. This gives smooth input (circular/spherical coverage). On a 1d axis input this works identically to Axial.
+	Radial,
+};
+
+/** Dead Zone
+    *  Input values within the range LowerThreshold -> UpperThreshold will be remapped from 0 -> 1.
+	*  Values outside this range will be clamped.
+	*/
+UCLASS(NotBlueprintable, MinimalAPI, meta = (DisplayName = "Dead Zone"))
+class UInputModifierDeadZone : public UInputModifier
+{
+	GENERATED_BODY()
+
+public:
+
+	// Threshold below which input is ignored
+	UPROPERTY(EditInstanceOnly, BlueprintReadWrite, Category=Settings, Config)
+	float LowerThreshold = 0.2f;
+
+	// Threshold above which input is clamped to 1
+	UPROPERTY(EditInstanceOnly, BlueprintReadWrite, Category=Settings, Config)
+	float UpperThreshold = 1.f;
+
+	UPROPERTY(EditInstanceOnly, BlueprintReadWrite, Category=Settings, Config)
+	EDeadZoneType Type = EDeadZoneType::Radial;
+
+protected:
+	virtual FInputActionValue ModifyRaw_Implementation(const UEnhancedPlayerInput* PlayerInput, FInputActionValue CurrentValue, float DeltaTime) override;
+
+	// Visualize as black when unmodified. Red when blocked (with differing intensities to indicate axes)
+	// Mirrors visualization in https://www.gamasutra.com/blogs/JoshSutphin/20130416/190541/Doing_Thumbstick_Dead_Zones_Right.php.
+	virtual FLinearColor GetVisualizationColor_Implementation(FInputActionValue SampleValue, FInputActionValue FinalValue) const override;
+};
+
+/** Scalar
+	*  Scales input by a set factor per axis
+	*/
+UCLASS(NotBlueprintable, MinimalAPI, meta = (DisplayName = "Scalar"))
+class UInputModifierScalar : public UInputModifier
+{
+	GENERATED_BODY()
+
+public:
+
+	// TODO: Detail customization to only show modifiable axes for the relevant binding? This thing has no idea what it's bound to...
+	UPROPERTY(EditInstanceOnly, BlueprintReadWrite, Category=Settings)
+	FVector Scalar = FVector::OneVector;
+
+protected:
+	virtual FInputActionValue ModifyRaw_Implementation(const UEnhancedPlayerInput* PlayerInput, FInputActionValue CurrentValue, float DeltaTime) override;
+};
+
+/** Negate
+	*  Inverts input per axis
+	*/
+UCLASS(NotBlueprintable, MinimalAPI, meta = (DisplayName = "Negate"))
+class UInputModifierNegate : public UInputModifier
+{
+	GENERATED_BODY()
+
+public:
+
+	UPROPERTY(EditInstanceOnly, BlueprintReadWrite, Category = Settings)
+	bool bX = true;
+	UPROPERTY(EditInstanceOnly, BlueprintReadWrite, Category = Settings)
+	bool bY = true;
+	UPROPERTY(EditInstanceOnly, BlueprintReadWrite, Category = Settings)
+	bool bZ = true;
+
+protected:
+	virtual FInputActionValue ModifyRaw_Implementation(const UEnhancedPlayerInput* PlayerInput, FInputActionValue CurrentValue, float DeltaTime) override;
+	virtual FLinearColor GetVisualizationColor_Implementation(FInputActionValue SampleValue, FInputActionValue FinalValue) const override;
+};
+
+/** Smooth
+	*  Smooth inputs out over multiple frames
+	*/
+UCLASS(NotBlueprintable, MinimalAPI, meta = (DisplayName = "Smooth"))
+class UInputModifierSmooth : public UInputModifier
+{
+	GENERATED_BODY()
+
+	void ClearSmoothedAxis();
+
+protected:
+
+	// TODO: Smoothing variants. Configuration options. e.g. smooth over a set time/frame count.
+
+	/** How long input has been zero. */
+	float ZeroTime = 0.f;
+
+	/** Current average input/sample */
+	FInputActionValue AverageValue;
+
+	/** Number of samples since input  has been zero */
+	int32 Samples = 0;
+
+#define SMOOTH_TOTAL_SAMPLE_TIME_DEFAULT (0.0083f)
+	/** Input sampling total time. */
+	float TotalSampleTime = SMOOTH_TOTAL_SAMPLE_TIME_DEFAULT;
+
+	virtual FInputActionValue ModifyRaw_Implementation(const UEnhancedPlayerInput* PlayerInput, FInputActionValue CurrentValue, float DeltaTime) override;
+};
+
+/** Response Curve Exponential
+	*  Apply a simple exponential response curve to input values, per axis
+	*/
+UCLASS(NotBlueprintable, MinimalAPI, meta = (DisplayName = "Response Curve - Exponential"))
+class UInputModifierResponseCurveExponential : public UInputModifier
+{
+	GENERATED_BODY()
+
+public:
+
+	UPROPERTY(EditInstanceOnly, BlueprintReadWrite, Category = Settings, Config)
+	FVector CurveExponent = FVector::OneVector;
+
+protected:
+	virtual FInputActionValue ModifyRaw_Implementation(const UEnhancedPlayerInput* PlayerInput, FInputActionValue CurrentValue, float DeltaTime) override;
+};
+
+
+/** Response Curve User Defined
+	*  Apply a custom response curve to input values, per axis
+	*/
+UCLASS(NotBlueprintable, MinimalAPI, meta = (DisplayName = "Response Curve - User Defined"))
+class UInputModifierResponseCurveUser : public	UInputModifier
+{
+	GENERATED_BODY()
+
+public:
+
+	UPROPERTY(EditInstanceOnly, BlueprintReadWrite, Category = Settings, meta = (DisplayThumbnail = "false"))
+	class UCurveFloat* ResponseX;
+	UPROPERTY(EditInstanceOnly, BlueprintReadWrite, Category = Settings, meta = (DisplayThumbnail = "false"))
+	class UCurveFloat* ResponseY;
+	UPROPERTY(EditInstanceOnly, BlueprintReadWrite, Category = Settings, meta = (DisplayThumbnail = "false"))
+	class UCurveFloat* ResponseZ;
+
+protected:
+	virtual FInputActionValue ModifyRaw_Implementation(const UEnhancedPlayerInput* PlayerInput, FInputActionValue CurrentValue, float DeltaTime) override;
+};
+
+UENUM()
+enum class EFOVScalingType : uint8
+{
+	// FOV scaling to apply scaled movement deltas to inputs dependent upon the player's selected FOV
+	Standard,
+
+	// FOV scaling was incorrectly calculated in UE4's UPlayerInput::MassageAxisInput. This implementation is intended to aid backwards compatibility, but should not be used by new projects.
+	UE4_BackCompat,
+};
+
+/** FOV Scaling
+	* Apply FOV dependent scaling to input values, per axis
+	*/
+UCLASS(NotBlueprintable, MinimalAPI, meta = (DisplayName = "FOV Scaling"))
+class UInputModifierFOVScaling : public	UInputModifier
+{
+	GENERATED_BODY()
+
+public:
+
+	// Extra scalar applied on top of basic FOV scaling.
+	UPROPERTY(EditInstanceOnly, BlueprintReadWrite, Category = Settings, Config)
+	float FOVScale = 1.f;
+
+	UPROPERTY(EditInstanceOnly, BlueprintReadWrite, Category = Settings, Config)
+	EFOVScalingType FOVScalingType = EFOVScalingType::Standard;	// TODO: UE4_BackCompat by default?
+
+protected:
+	virtual FInputActionValue ModifyRaw_Implementation(const UEnhancedPlayerInput* PlayerInput, FInputActionValue CurrentValue, float DeltaTime) override;
+};
+
+
+/** Input space to World space conversion
+	* Auto-converts axes within the Input Action Value into world space	allowing the result to be directly plugged into functions that take world space values.
+	* E.g. For a 2D input axis up/down is mapped to world X (forward), whilst axis left/right is mapped to world Y (right).
+	*/
+UCLASS(NotBlueprintable, MinimalAPI, meta = (DisplayName = "To World Space"))
+class UInputModifierToWorldSpace : public	UInputModifier
+{
+	GENERATED_BODY()
+
+protected:
+	virtual FInputActionValue ModifyRaw_Implementation(const UEnhancedPlayerInput* PlayerInput, FInputActionValue CurrentValue, float DeltaTime) override;
+	virtual FLinearColor GetVisualizationColor_Implementation(FInputActionValue SampleValue, FInputActionValue FinalValue) const override;
+};
+
+
+UENUM()
+enum class EInputAxisSwizzle : uint8
+{
+	// Swap X and Y axis. Useful for binding 1D inputs to the Y axis for 2D actions.
+	YXZ,
+
+	// Swap X and Z axis
+	ZYX,
+
+	// Swap Y and Z axis
+	XZY,
+
+	// Reorder all axes, Y first
+	YZX,
+
+	// Reorder all axes, Z first
+	ZXY,
+};
+
+/** Swizzle axis components of an input value.
+	* Useful to map a 1D input onto the Y axis of a 2D action.
+	*/
+UCLASS(NotBlueprintable, MinimalAPI, meta = (DisplayName = "Swizzle Input Axis Values"))
+class UInputModifierSwizzleAxis : public UInputModifier
+{
+	GENERATED_BODY()
+
+public:
+
+	// Default to XY swap, useful for binding 1D inputs to the Y axis.
+	UPROPERTY(EditAnywhere, BlueprintReadWrite, Category = Settings)
+	EInputAxisSwizzle Order = EInputAxisSwizzle::YXZ;
+
+protected:
+	virtual FInputActionValue ModifyRaw_Implementation(const UEnhancedPlayerInput* PlayerInput, FInputActionValue CurrentValue, float DeltaTime) override;
+	virtual FLinearColor GetVisualizationColor_Implementation(FInputActionValue SampleValue, FInputActionValue FinalValue) const override;
+};
+
+/** Modifier collection
+	* A user definable group of modifiers that can be easily applied to multiple actions or mappings to save duplication work.
+	*/
+UCLASS(NotBlueprintable, MinimalAPI, meta = (DisplayName = "Modifier Collection"))
+class UInputModifierCollection : public	UInputModifier
+{
+	GENERATED_BODY()
+
+public:
+
+	UPROPERTY(EditInstanceOnly, BlueprintReadWrite, Category = Settings)
+	TArray<UInputModifier*> Modifiers;	// TODO: These either need to be instanced per copy of the collection or it needs to be clear the modifiers are okay running multiple times per frame (or we just don't support collections)
+
+	/** If set each modifier will not have the modified value corrected to the base type before execution.
+	*	After all modifiers are run the resulting value will be converted back to the action's value type as with any other modifier.
+	*	This allows for complex sets of conditional modifiers that can alter their behavior based on their predecessors value type.
+	*	Note that this is an advanced feature and may cause issues if used with the basic modifier implementations.
+	**/
+	UPROPERTY(EditAnywhere, BlueprintReadWrite, Category = Settings)
+	bool bPermitValueTypeModification = false;
+
+protected:
+	virtual FInputActionValue ModifyRaw_Implementation(const UEnhancedPlayerInput* PlayerInput, FInputActionValue CurrentValue, float DeltaTime) override;
+};
\ No newline at end of file
diff --git a/Engine/Plugins/EnhancedInput/Source/InputBlueprintNodes/Private/InputBlueprintNodesModule.cpp b/Engine/Plugins/EnhancedInput/Source/InputBlueprintNodes/Private/InputBlueprintNodesModule.cpp
new file mode 100644
--- /dev/null
+++ b/Engine/Plugins/EnhancedInput/Source/InputBlueprintNodes/Private/InputBlueprintNodesModule.cpp
@@ -0,0 +1,17 @@
+// Copyright Epic Games, Inc. All Rights Reserved.
+
+#include "InputBlueprintNodesModule.h"
+
+#include "CoreMinimal.h"
+#include "Modules/ModuleInterface.h"
+#include "Modules/ModuleManager.h"
+
+#define LOCTEXT_NAMESPACE "InputBlueprintNodes"
+
+class FInputBlueprintNodesModule : public IModuleInterface
+{
+};
+
+IMPLEMENT_MODULE(FInputBlueprintNodesModule, InputBlueprintNodes)
+
+#undef LOCTEXT_NAMESPACE
\ No newline at end of file
diff --git a/Engine/Plugins/EnhancedInput/Source/InputBlueprintNodes/Private/K2Node_EnhancedInputActionEvent.h b/Engine/Plugins/EnhancedInput/Source/InputBlueprintNodes/Private/K2Node_EnhancedInputActionEvent.h
new file mode 100644
--- /dev/null
+++ b/Engine/Plugins/EnhancedInput/Source/InputBlueprintNodes/Private/K2Node_EnhancedInputActionEvent.h
@@ -0,0 +1,29 @@
+// Copyright Epic Games, Inc. All Rights Reserved.
+
+#pragma once
+
+#include "CoreMinimal.h"
+#include "UObject/ObjectMacros.h"
+#include "Engine/EngineBaseTypes.h"
+#include "EnhancedPlayerInput.h"
+#include "K2Node_Event.h"
+#include "K2Node_EnhancedInputActionEvent.generated.h"
+
+class UDynamicBlueprintBinding;
+
+UCLASS()
+class INPUTBLUEPRINTNODES_API UK2Node_EnhancedInputActionEvent : public UK2Node_Event
+{
+	GENERATED_UCLASS_BODY()
+
+	UPROPERTY()
+	const UInputAction* InputAction;
+
+	UPROPERTY()
+	ETriggerEvent TriggerEvent;
+
+	//~ Begin UK2Node Interface
+	virtual UClass* GetDynamicBindingClass() const override;
+	virtual void RegisterDynamicBinding(UDynamicBlueprintBinding* BindingObject) const override;
+	//~ End UK2Node Interface
+};
diff --git a/Engine/Plugins/EnhancedInput/Source/InputBlueprintNodes/Private/K2Node_InputActionValueAccessor.h b/Engine/Plugins/EnhancedInput/Source/InputBlueprintNodes/Private/K2Node_InputActionValueAccessor.h
new file mode 100644
--- /dev/null
+++ b/Engine/Plugins/EnhancedInput/Source/InputBlueprintNodes/Private/K2Node_InputActionValueAccessor.h
@@ -0,0 +1,33 @@
+// Copyright Epic Games, Inc. All Rights Reserved.
+
+#pragma once
+
+#include "CoreMinimal.h"
+#include "UObject/ObjectMacros.h"
+#include "Engine/EngineBaseTypes.h"
+#include "InputAction.h"
+#include "K2Node_CallFunction.h"
+#include "K2Node_InputActionValueAccessor.generated.h"
+
+UCLASS()
+class INPUTBLUEPRINTNODES_API UK2Node_InputActionValueAccessor : public UK2Node_CallFunction
+{
+	GENERATED_UCLASS_BODY()
+
+public:
+	//~ Begin EdGraphNode Interface
+	virtual void AllocateDefaultPins() override;
+	//~ End EdGraphNode Interface
+
+	//~ Begin UK2Node Interface
+	virtual UClass* GetDynamicBindingClass() const override;
+	virtual void RegisterDynamicBinding(class UDynamicBlueprintBinding* BindingObject) const override;
+	virtual bool IsNodePure() const override { return true; }
+	//~ End UK2Node Interface
+
+	void Initialize(const UInputAction* Action);
+
+private:
+	UPROPERTY()
+	const UInputAction* InputAction;
+};
diff --git a/Engine/Plugins/EnhancedInput/Source/InputBlueprintNodes/Private/K2Node_InputDebugKey.h b/Engine/Plugins/EnhancedInput/Source/InputBlueprintNodes/Private/K2Node_InputDebugKey.h
new file mode 100644
--- /dev/null
+++ b/Engine/Plugins/EnhancedInput/Source/InputBlueprintNodes/Private/K2Node_InputDebugKey.h
@@ -0,0 +1,95 @@
+// Copyright Epic Games, Inc. All Rights Reserved.
+
+#pragma once
+
+#include "CoreMinimal.h"
+#include "UObject/ObjectMacros.h"
+#include "InputCoreTypes.h"
+#include "EdGraph/EdGraph.h"
+#include "Textures/SlateIcon.h"
+#include "EdGraph/EdGraphNodeUtils.h"
+#include "BlueprintNodeSignature.h"
+#include "K2Node.h"
+#include "K2Node_EventNodeInterface.h"
+
+#include "K2Node_InputDebugKey.generated.h"
+
+class FBlueprintActionDatabaseRegistrar;
+class FKismetCompilerContext;
+
+UCLASS()
+class INPUTBLUEPRINTNODES_API UK2Node_InputDebugKey : public UK2Node, public IK2Node_EventNodeInterface
+{
+	GENERATED_UCLASS_BODY()
+
+	// The key that is bound
+	UPROPERTY(EditAnywhere, Category="Input")
+	FKey InputKey;
+
+	// Should the binding execute even when the game is paused
+	UPROPERTY(EditAnywhere, Category="Input")
+	uint32 bExecuteWhenPaused:1;
+
+	// Does this binding require the control key on PC or the command key on Mac to be held
+	UPROPERTY(EditAnywhere, Category="Modifier")
+	uint32 bControl:1;
+
+	// Does this binding require the alt key to be held
+	UPROPERTY(EditAnywhere, Category="Modifier")
+	uint32 bAlt:1;
+
+	// Does this binding require the shift key to be held
+	UPROPERTY(EditAnywhere, Category="Modifier")
+	uint32 bShift:1;
+
+	// Does this binding require the windows key on PC or the control key on Mac to be held
+	UPROPERTY(EditAnywhere, Category="Modifier")
+	uint32 bCommand:1;
+
+	//~ Begin UObject Interface
+	virtual void PostEditChangeProperty(struct FPropertyChangedEvent& PropertyChangedEvent) override;
+	//~ End UObject Interface
+
+	//~ Begin UK2Node Interface.
+	virtual bool ShouldShowNodeProperties() const override { return true; }
+	virtual void ValidateNodeDuringCompilation(class FCompilerResultsLog& MessageLog) const override;
+	virtual void ExpandNode(FKismetCompilerContext& CompilerContext, UEdGraph* SourceGraph) override;
+	virtual void GetMenuActions(FBlueprintActionDatabaseRegistrar& ActionRegistrar) const override;
+	virtual FText GetMenuCategory() const override;
+	virtual FBlueprintNodeSignature GetSignature() const override;
+	//~ End UK2Node Interface
+
+	//~ Begin UEdGraphNode Interface.
+	virtual void AllocateDefaultPins() override;
+	virtual FLinearColor GetNodeTitleColor() const override;
+	virtual FText GetNodeTitle(ENodeTitleType::Type TitleType) const override;
+	virtual FText GetTooltipText() const override;
+	virtual FSlateIcon GetIconAndTint(FLinearColor& OutColor) const override;
+	virtual bool IsCompatibleWithGraph(UEdGraph const* Graph) const override;
+	//~ End UEdGraphNode Interface.
+
+	//~ Begin IK2Node_EventNodeInterface Interface.
+	virtual TSharedPtr<FEdGraphSchemaAction> GetEventNodeAction(const FText& ActionCategory) override;
+	//~ End IK2Node_EventNodeInterface Interface.
+
+
+	FText GetModifierText() const;
+
+	FName GetModifierName() const;
+
+	FText GetKeyText() const;
+
+	/** Get the 'pressed' input pin */
+	UEdGraphPin* GetPressedPin() const;
+
+	/** Get the 'released' input pin */
+	UEdGraphPin* GetReleasedPin() const;
+
+	/** Get the 'Action Value' input pin */
+	UEdGraphPin* GetActionValuePin() const;
+
+private:
+	/** Constructing FText strings can be costly, so we cache the node's title/tooltip */
+	FNodeTextCache CachedTooltip;
+	FNodeTextCache CachedNodeTitle;
+};
diff --git a/Engine/Plugins/EnhancedInput/Source/InputBlueprintNodes/Private/K2Node_InputDebugKeyEvent.h b/Engine/Plugins/EnhancedInput/Source/InputBlueprintNodes/Private/K2Node_InputDebugKeyEvent.h
new file mode 100644
--- /dev/null
+++ b/Engine/Plugins/EnhancedInput/Source/InputBlueprintNodes/Private/K2Node_InputDebugKeyEvent.h
@@ -0,0 +1,32 @@
+// Copyright Epic Games, Inc. All Rights Reserved.
+
+#pragma once
+
+#include "CoreMinimal.h"
+#include "UObject/ObjectMacros.h"
+#include "Engine/EngineBaseTypes.h"
+#include "Framework/Commands/InputChord.h"
+#include "K2Node_Event.h"
+#include "K2Node_InputDebugKeyEvent.generated.h"
+
+class UDynamicBlueprintBinding;
+
+UCLASS()
+class INPUTBLUEPRINTNODES_API UK2Node_InputDebugKeyEvent : public UK2Node_Event
+{
+	GENERATED_UCLASS_BODY()
+
+	UPROPERTY()
+	FInputChord InputChord;
+
+	UPROPERTY()
+	TEnumAsByte<EInputEvent> InputKeyEvent;
+
+	UPROPERTY()
+	bool bExecuteWhenPaused = false;
+
+	//~ Begin UK2Node Interface
+	virtual UClass* GetDynamicBindingClass() const override;
+	virtual void RegisterDynamicBinding(UDynamicBlueprintBinding* BindingObject) const override;
+	//~ End UK2Node Interface
+};
diff --git a/Engine/Plugins/EnhancedInput/Source/InputBlueprintNodes/Public/InputBlueprintNodesModule.h b/Engine/Plugins/EnhancedInput/Source/InputBlueprintNodes/Public/InputBlueprintNodesModule.h
new file mode 100644
--- /dev/null
+++ b/Engine/Plugins/EnhancedInput/Source/InputBlueprintNodes/Public/InputBlueprintNodesModule.h
@@ -0,0 +1,10 @@
+// Copyright Epic Games, Inc. All Rights Reserved.
+
+#pragma once
+
+//#include "InputBlueprintNodesModule.generated.h"
+
+// Module is not publicly exposed
+
+
+
diff --git a/Engine/Plugins/EnhancedInput/Source/InputBlueprintNodes/Public/K2Node_EnhancedInputAction.h b/Engine/Plugins/EnhancedInput/Source/InputBlueprintNodes/Public/K2Node_EnhancedInputAction.h
new file mode 100644
--- /dev/null
+++ b/Engine/Plugins/EnhancedInput/Source/InputBlueprintNodes/Public/K2Node_EnhancedInputAction.h
@@ -0,0 +1,66 @@
+// Copyright Epic Games, Inc. All Rights Reserved.
+
+#pragma once
+
+#include "CoreMinimal.h"
+#include "BlueprintNodeSignature.h"
+#include "EdGraph/EdGraph.h"
+#include "EdGraph/EdGraphNodeUtils.h"
+#include "InputAction.h"
+#include "K2Node.h"
+#include "K2Node_EventNodeInterface.h"
+#include "Textures/SlateIcon.h"
+#include "UObject/ObjectMacros.h"
+
+#include "K2Node_EnhancedInputAction.generated.h"
+
+class FBlueprintActionDatabaseRegistrar;
+class FKismetCompilerContext;
+
+UCLASS()
+class INPUTBLUEPRINTNODES_API UK2Node_EnhancedInputAction : public UK2Node, public IK2Node_EventNodeInterface
+{
+	GENERATED_UCLASS_BODY()
+
+	UPROPERTY()
+	const UInputAction* InputAction;
+
+	//~ Begin UEdGraphNode Interface.
+	virtual void AllocateDefaultPins() override;
+	virtual FLinearColor GetNodeTitleColor() const override;
+	virtual FText GetNodeTitle(ENodeTitleType::Type TitleType) const override;
+	virtual FText GetTooltipText() const override;
+	virtual FSlateIcon GetIconAndTint(FLinearColor& OutColor) const override;
+	virtual bool IsCompatibleWithGraph(UEdGraph const* Graph) const override;
+	virtual UObject* GetJumpTargetForDoubleClick() const override;
+	virtual void JumpToDefinition() const override;
+	//~ End UEdGraphNode Interface.
+
+	//~ Begin UK2Node Interface
+	virtual void ValidateNodeDuringCompilation(class FCompilerResultsLog& MessageLog) const override;
+	virtual bool ShouldShowNodeProperties() const override { return true; }
+	virtual void ExpandNode(FKismetCompilerContext& CompilerContext, UEdGraph* SourceGraph) override;
+	virtual void GetMenuActions(FBlueprintActionDatabaseRegistrar& ActionRegistrar) const override;
+	virtual FText GetMenuCategory() const override;
+	virtual bool CanUserEditPinAdvancedViewFlag() const override { return true; }
+	virtual FBlueprintNodeSignature GetSignature() const override;
+	virtual void PostReconstructNode();
+	virtual void PinConnectionListChanged(UEdGraphPin* Pin);
+	virtual bool IsConnectionDisallowed(const UEdGraphPin* MyPin, const UEdGraphPin* OtherPin, FString& OutReason) const override;
+	//~ End UK2Node Interface
+
+	//~ Begin IK2Node_EventNodeInterface Interface.
+	virtual TSharedPtr<FEdGraphSchemaAction> GetEventNodeAction(const FText& ActionCategory) override;
+	//~ End IK2Node_EventNodeInterface Interface.
+
+private:
+	FName GetActionName() const;
+	void HideEventPins(UEdGraphPin* RetainPin);
+
+	/** Gets the ETriggerEvent from an exec pin based on the Pins name. */
+	const ETriggerEvent GetTriggerTypeFromExecPin(const UEdGraphPin* ExecPin) const;
+	
+	/** Constructing FText strings can be costly, so we cache the node's title/tooltip */
+	FNodeTextCache CachedTooltip;
+	FNodeTextCache CachedNodeTitle;
+};
diff --git a/Engine/Plugins/EnhancedInput/Source/InputEditor/Private/InputCustomizations.cpp b/Engine/Plugins/EnhancedInput/Source/InputEditor/Private/InputCustomizations.cpp
new file mode 100644
--- /dev/null
+++ b/Engine/Plugins/EnhancedInput/Source/InputEditor/Private/InputCustomizations.cpp
@@ -0,0 +1,77 @@
+// Copyright Epic Games, Inc. All Rights Reserved.
+
+#include "InputCustomizations.h"
+
+#include "ActionMappingDetails.h"
+#include "DetailCategoryBuilder.h"
+#include "DetailLayoutBuilder.h"
+#include "EnhancedActionKeyMapping.h"
+#include "IDetailChildrenBuilder.h"
+#include "InputMappingContext.h"
+#include "KeyStructCustomization.h"
+#include "PropertyCustomizationHelpers.h"
+
+#define LOCTEXT_NAMESPACE "InputCustomization"
+
+TSharedRef<IDetailCustomization> FInputContextDetails::MakeInstance()
+{
+	return MakeShareable(new FInputContextDetails);
+}
+
+void FInputContextDetails::CustomizeDetails(class IDetailLayoutBuilder& DetailBuilder)
+{
+	// Custom Action Mappings
+	const TSharedPtr<IPropertyHandle> ActionMappingsPropertyHandle = DetailBuilder.GetProperty(GET_MEMBER_NAME_CHECKED(UInputMappingContext, Mappings));
+	ActionMappingsPropertyHandle->MarkHiddenByCustomization();
+
+	IDetailCategoryBuilder& MappingsDetailCategoryBuilder = DetailBuilder.EditCategory(ActionMappingsPropertyHandle->GetDefaultCategoryName());
+	const TSharedRef<FActionMappingsNodeBuilderEx> ActionMappingsBuilder = MakeShareable(new FActionMappingsNodeBuilderEx(&DetailBuilder, ActionMappingsPropertyHandle));
+	MappingsDetailCategoryBuilder.AddCustomBuilder(ActionMappingsBuilder);
+}
+
+
+void FEnhancedActionMappingCustomization::CustomizeHeader(TSharedRef<IPropertyHandle> PropertyHandle, FDetailWidgetRow& HeaderRow, IPropertyTypeCustomizationUtils& CustomizationUtils)
+{
+	MappingPropertyHandle = PropertyHandle;
+
+	// Grab the FKey property
+	TSharedPtr<IPropertyHandle> KeyHandle = PropertyHandle->GetChildHandle(GET_MEMBER_NAME_CHECKED(FEnhancedActionKeyMapping, Key));
+
+	TSharedRef<SWidget> RemoveButton = PropertyCustomizationHelpers::MakeDeleteButton(FSimpleDelegate::CreateSP(this, &FEnhancedActionMappingCustomization::RemoveMappingButton_OnClick),
+		LOCTEXT("RemoveMappingToolTip", "Remove Mapping"));
+
+	// Create  a new instance of the key customization.
+	KeyStructInstance = FKeyStructCustomization::MakeInstance();
+
+	// TODO: Use FDetailArrayBuilder?
+
+	// Pass our header row into the key struct customizeheader method so it populates our row with the key struct header
+	StaticCastSharedPtr<FKeyStructCustomization>(KeyStructInstance)->CustomizeHeaderOnlyWithButton(KeyHandle.ToSharedRef(), HeaderRow, CustomizationUtils, RemoveButton);
+}
+
+void FEnhancedActionMappingCustomization::CustomizeChildren(TSharedRef<IPropertyHandle> PropertyHandle, IDetailChildrenBuilder& ChildBuilder, IPropertyTypeCustomizationUtils& CustomizationUtils)
+{
+	TSharedPtr<IPropertyHandle> TriggersHandle = PropertyHandle->GetChildHandle(GET_MEMBER_NAME_CHECKED(FEnhancedActionKeyMapping, Triggers));
+	TSharedPtr<IPropertyHandle> ModifiersHandle = PropertyHandle->GetChildHandle(GET_MEMBER_NAME_CHECKED(FEnhancedActionKeyMapping, Modifiers));
+	TSharedPtr<IPropertyHandle> IsPlayerMappableHandle = PropertyHandle->GetChildHandle(GET_MEMBER_NAME_CHECKED(FEnhancedActionKeyMapping, bIsPlayerMappable));
+	TSharedPtr<IPropertyHandle> PlayerBindingOptions = PropertyHandle->GetChildHandle(GET_MEMBER_NAME_CHECKED(FEnhancedActionKeyMapping, PlayerMappableOptions));
+
+	// TODO: ResetToDefault needs to be disabled for arrays
+	ChildBuilder.AddProperty(TriggersHandle.ToSharedRef());
+	ChildBuilder.AddProperty(ModifiersHandle.ToSharedRef());
+	ChildBuilder.AddProperty(IsPlayerMappableHandle.ToSharedRef());
+	ChildBuilder.AddProperty(PlayerBindingOptions.ToSharedRef());
+}
+
+void FEnhancedActionMappingCustomization::RemoveMappingButton_OnClick() const
+{
+	if (MappingPropertyHandle->IsValidHandle())
+	{
+		const TSharedPtr<IPropertyHandle> ParentHandle = MappingPropertyHandle->GetParentHandle();
+		const TSharedPtr<IPropertyHandleArray> ParentArrayHandle = ParentHandle->AsArray();
+
+		ParentArrayHandle->DeleteItem(MappingPropertyHandle->GetIndexInArray());
+	}
+}
+
+#undef LOCTEXT_NAMESPACE
\ No newline at end of file
diff --git a/Engine/Plugins/EnhancedInput/Source/InputEditor/Private/InputCustomizations.h b/Engine/Plugins/EnhancedInput/Source/InputEditor/Private/InputCustomizations.h
new file mode 100644
--- /dev/null
+++ b/Engine/Plugins/EnhancedInput/Source/InputEditor/Private/InputCustomizations.h
@@ -0,0 +1,38 @@
+// Copyright Epic Games, Inc. All Rights Reserved.
+
+#pragma once
+
+#include "CoreMinimal.h"
+#include "IDetailCustomization.h"
+#include "IDetailCustomNodeBuilder.h"
+#include "IPropertyTypeCustomization.h"
+
+class FInputContextDetails : public IDetailCustomization
+{
+public:
+	/** Makes a new instance of this detail layout class for a specific detail view requesting it */
+	static TSharedRef<IDetailCustomization> MakeInstance();
+
+	/** ILayoutDetails interface */
+	virtual void CustomizeDetails( class IDetailLayoutBuilder& DetailBuilder ) override;
+};
+
+class FEnhancedActionMappingCustomization : public IPropertyTypeCustomization
+{
+public:
+	static TSharedRef<IPropertyTypeCustomization> MakeInstance()
+	{
+		return MakeShareable(new FEnhancedActionMappingCustomization());
+	}
+
+	/** IPropertyTypeCustomization interface */
+	virtual void CustomizeHeader(TSharedRef<IPropertyHandle> PropertyHandle, FDetailWidgetRow& HeaderRow, IPropertyTypeCustomizationUtils& CustomizationUtils) override;
+	virtual void CustomizeChildren(TSharedRef<IPropertyHandle> PropertyHandle, IDetailChildrenBuilder& ChildBuilder, IPropertyTypeCustomizationUtils& CustomizationUtils) override;
+
+private:
+
+	void RemoveMappingButton_OnClick() const;
+
+	TSharedPtr<IPropertyTypeCustomization> KeyStructInstance;
+	TSharedPtr<IPropertyHandle> MappingPropertyHandle;
+};
diff --git a/Engine/Plugins/EnhancedInput/Source/InputEditor/Private/InputEditorModule.cpp b/Engine/Plugins/EnhancedInput/Source/InputEditor/Private/InputEditorModule.cpp
new file mode 100644
--- /dev/null
+++ b/Engine/Plugins/EnhancedInput/Source/InputEditor/Private/InputEditorModule.cpp
@@ -0,0 +1,570 @@
+// Copyright Epic Games, Inc. All Rights Reserved.
+
+#include "InputEditorModule.h"
+
+#include "AssetRegistryModule.h"
+#include "AssetTypeActions_Base.h"
+#include "EnhancedInputModule.h"
+#include "DetailCategoryBuilder.h"
+#include "DetailLayoutBuilder.h"
+#include "Framework/Notifications/NotificationManager.h"
+#include "GameFramework/PlayerController.h"
+#include "HAL/PlatformFileManager.h"
+#include "InputAction.h"
+#include "InputMappingContext.h"
+#include "PlayerMappableInputConfig.h"
+#include "InputCustomizations.h"
+#include "InputModifiers.h"
+#include "IAssetTools.h"
+#include "IAssetTypeActions.h"
+#include "IDetailsView.h"
+#include "ISettingsModule.h"
+#include "K2Node_EnhancedInputAction.h"
+#include "K2Node_GetInputActionValue.h"
+#include "Modules/ModuleInterface.h"
+#include "Modules/ModuleManager.h"
+#include "PropertyEditorDelegates.h"
+#include "PropertyEditorModule.h"
+#include "SSettingsEditorCheckoutNotice.h"
+#include "TickableEditorObject.h"
+#include "Widgets/Layout/SScrollBox.h"
+#include "Widgets/Notifications/SNotificationList.h"
+#include "AssetTypeActions/AssetTypeActions_DataAsset.h"
+
+#define LOCTEXT_NAMESPACE "InputEditor"
+
+class FInputEditorModule : public IModuleInterface, public FTickableEditorObject
+{
+public:
+
+	// IModuleInterface interface
+	virtual void StartupModule() override;
+	virtual void ShutdownModule() override;
+	// End IModuleInterface interface
+
+	// FTickableEditorObject interface
+	virtual void Tick(float DeltaTime) override;
+	virtual TStatId GetStatId() const override { RETURN_QUICK_DECLARE_CYCLE_STAT(FInputEditorModule, STATGROUP_Tickables); }
+	// End FTickableEditorObject interface
+
+	static EAssetTypeCategories::Type GetInputAssetsCategory() { return InputAssetsCategory; }
+
+private:
+	void RegisterAssetTypeActions(IAssetTools& AssetTools, TSharedRef<IAssetTypeActions> Action)
+	{
+		AssetTools.RegisterAssetTypeActions(Action);
+		CreatedAssetTypeActions.Add(Action);
+	}
+	void PostEngineInit();
+	TSharedRef<SWidget> CreateSettingsPanel();
+	void OnSettingChanged(const FPropertyChangedEvent& PropertyChangedEvent);
+
+	void RebuildDetailsViewForAsset(const FAssetData& AssetData, bool bIgnoreAsset);
+	void OnAssetAdded(const FAssetData& AssetData) { RebuildDetailsViewForAsset(AssetData, false); }
+	void OnAssetRemoved(const FAssetData& AssetData) { RebuildDetailsViewForAsset(AssetData, true); }
+	void OnAssetRenamed(const FAssetData& AssetData, const FString&) { RebuildDetailsViewForAsset(AssetData, true); }
+
+	template<typename T>
+	TSharedPtr<IDetailsView> AddClassDetailsView();
+
+	struct FClassDetailsView
+	{
+		FClassDetailsView() = default;
+		FClassDetailsView(UClass* InClass, TSharedPtr<IDetailsView>& InView) : Class(InClass), View(InView) {}
+		bool IsValid() const { return Class != nullptr; }
+
+		UClass* Class = nullptr;
+		TSharedPtr<IDetailsView> View;
+	};
+
+	FClassDetailsView FindClassDetailsViewForAsset(const FAssetData& AssetData);
+	TArray<UObject*> GatherClassDetailsCDOs(UClass* Class, const FAssetData* IgnoreAsset);
+
+	static EAssetTypeCategories::Type InputAssetsCategory;
+
+	TArray< TSharedPtr<IAssetTypeActions> > CreatedAssetTypeActions;
+
+	TMap<UClass*, TSharedPtr<IDetailsView>> DetailsViews;
+
+	TSharedPtr<SWidget> Panel;
+};
+
+EAssetTypeCategories::Type FInputEditorModule::InputAssetsCategory;
+
+IMPLEMENT_MODULE(FInputEditorModule, InputEditor)
+
+// Asset factories
+
+// InputContext
+UInputMappingContext_Factory::UInputMappingContext_Factory(const class FObjectInitializer& OBJ) : Super(OBJ) {
+	SupportedClass = UInputMappingContext::StaticClass();
+	bEditAfterNew = true;
+	bCreateNew = true;
+}
+
+UObject* UInputMappingContext_Factory::FactoryCreateNew(UClass* Class, UObject* InParent, FName Name, EObjectFlags Flags, UObject* Context, FFeedbackContext* Warn)
+{
+	check(Class->IsChildOf(UInputMappingContext::StaticClass()));
+	return NewObject<UInputMappingContext>(InParent, Class, Name, Flags | RF_Transactional, Context);
+}
+
+// InputAction
+UInputAction_Factory::UInputAction_Factory(const class FObjectInitializer& OBJ)
+	: Super(OBJ)
+{
+	SupportedClass = UInputAction::StaticClass();
+	bEditAfterNew = true;
+	bCreateNew = true;
+}
+
+UObject* UInputAction_Factory::FactoryCreateNew(UClass* Class, UObject* InParent, FName Name, EObjectFlags Flags, UObject* Context, FFeedbackContext* Warn)
+{
+	check(Class->IsChildOf(UInputAction::StaticClass()));
+	return NewObject<UInputAction>(InParent, Class, Name, Flags | RF_Transactional, Context);
+}
+
+// UPlayerMappableInputConfig_Factory
+UPlayerMappableInputConfig_Factory::UPlayerMappableInputConfig_Factory(const class FObjectInitializer& OBJ)
+	: Super(OBJ)
+{
+	SupportedClass = UPlayerMappableInputConfig::StaticClass();
+	bEditAfterNew = true;
+	bCreateNew = true;
+}
+
+UObject* UPlayerMappableInputConfig_Factory::FactoryCreateNew(UClass* Class, UObject* InParent, FName Name, EObjectFlags Flags, UObject* Context, FFeedbackContext* Warn)
+{
+	check(Class->IsChildOf(UPlayerMappableInputConfig::StaticClass()));
+	return NewObject<UPlayerMappableInputConfig>(InParent, Class, Name, Flags | RF_Transactional, Context);
+}
+
+//
+//// InputTrigger
+//UInputTrigger_Factory::UInputTrigger_Factory(const class FObjectInitializer& OBJ) : Super(OBJ) {
+//	ParentClass = UInputTrigger::StaticClass();
+//	SupportedClass = UInputTrigger::StaticClass();
+//	bEditAfterNew = true;
+//	bCreateNew = true;
+//}
+//
+//// InputModifier
+//UInputModifier_Factory::UInputModifier_Factory(const class FObjectInitializer& OBJ) : Super(OBJ) {
+//	ParentClass = UInputModifier::StaticClass();
+//	SupportedClass = UInputModifier::StaticClass();
+//	bEditAfterNew = true;
+//	bCreateNew = true;
+//}
+
+
+
+// Asset type actions
+// TODO: Move asset type action definitions out?
+
+class FAssetTypeActions_InputContext : public FAssetTypeActions_DataAsset
+{
+public:
+	virtual FText GetName() const override { return NSLOCTEXT("AssetTypeActions", "AssetTypeActions_InputMappingContext", "Input Mapping Context"); }
+	virtual uint32 GetCategories() override { return FInputEditorModule::GetInputAssetsCategory(); }
+	virtual FColor GetTypeColor() const override { return FColor(255, 255, 127); }
+	virtual FText GetAssetDescription(const FAssetData& AssetData) const override { return NSLOCTEXT("AssetTypeActions", "AssetTypeActions_InputContextDesc", "A collection of device input to action mappings."); }
+	virtual UClass* GetSupportedClass() const override { return UInputMappingContext::StaticClass(); }
+};
+
+class FAssetTypeActions_InputAction : public FAssetTypeActions_DataAsset
+{
+public:
+	virtual FText GetName() const override { return NSLOCTEXT("AssetTypeActions", "AssetTypeActions_InputAction", "Input Action"); }
+	virtual uint32 GetCategories() override { return FInputEditorModule::GetInputAssetsCategory(); }
+	virtual FColor GetTypeColor() const override { return FColor(127, 255, 255); }
+	virtual FText GetAssetDescription(const FAssetData& AssetData) const override { return NSLOCTEXT("AssetTypeActions", "AssetTypeActions_InputActionDesc", "Represents an an abstract game action that can be mapped to arbitrary hardware input devices."); }
+	virtual UClass* GetSupportedClass() const override { return UInputAction::StaticClass(); }
+};
+
+class FAssetTypeActions_PlayerMappableInputConfig : public FAssetTypeActions_DataAsset
+{
+public:
+	virtual FText GetName() const override { return NSLOCTEXT("AssetTypeActions", "AssetTypeActions_PlayerBindableInputConfig", "Player Bindable Input Config"); }
+	virtual uint32 GetCategories() override { return FInputEditorModule::GetInputAssetsCategory(); }
+	virtual FColor GetTypeColor() const override { return FColor(127, 255, 255); }
+	virtual FText GetAssetDescription(const FAssetData& AssetData) const override { return NSLOCTEXT("AssetTypeActions", "AssetTypeActions_PlayerBindableInputConfigDesc", "Represents one set of Player Mappable controller/keymappings"); }
+	virtual UClass* GetSupportedClass() const override { return UPlayerMappableInputConfig::StaticClass(); }
+};
+
+
+
+void FInputEditorModule::OnSettingChanged(const FPropertyChangedEvent& PropertyChangedEvent)
+{
+	// TODO: Copy of SSettingsEditor::NotifyPostChange
+	if (PropertyChangedEvent.ChangeType != EPropertyChangeType::Interactive)
+	{
+		check(PropertyChangedEvent.GetNumObjectsBeingEdited() <= 1);
+		if(PropertyChangedEvent.GetNumObjectsBeingEdited() > 0)
+		{
+			UObject* ObjectBeingEdited = (UObject*)PropertyChangedEvent.GetObjectBeingEdited(0);
+
+			// Attempt to checkout the file automatically
+			if (!ObjectBeingEdited->GetClass()->HasAnyClassFlags(CLASS_DefaultConfig))
+			{
+				return;
+			}
+			FString RelativePath = ObjectBeingEdited->GetDefaultConfigFilename();
+
+			FString FullPath = FPaths::ConvertRelativePathToFull(RelativePath);
+
+			bool bIsNewFile = !FPlatformFileManager::Get().GetPlatformFile().FileExists(*FullPath);
+
+			if (!SettingsHelpers::CheckOutOrAddFile(FullPath))
+			{
+				SettingsHelpers::MakeWritable(FullPath);
+			}
+
+			// Determine if the Property is an Array or Array Element
+			bool bIsArrayOrArrayElement = PropertyChangedEvent.Property->IsA(FArrayProperty::StaticClass())
+				|| PropertyChangedEvent.Property->ArrayDim > 1
+				|| PropertyChangedEvent.Property->GetOwner<FArrayProperty>();
+
+			bool bIsSetOrSetElement = PropertyChangedEvent.Property->IsA(FSetProperty::StaticClass())
+				|| PropertyChangedEvent.Property->GetOwner<FSetProperty>();
+
+			bool bIsMapOrMapElement = PropertyChangedEvent.Property->IsA(FMapProperty::StaticClass())
+				|| PropertyChangedEvent.Property->GetOwner<FMapProperty>();
+
+			if (ObjectBeingEdited->GetClass()->HasAnyClassFlags(CLASS_DefaultConfig) && !bIsArrayOrArrayElement && !bIsSetOrSetElement && !bIsMapOrMapElement)
+			{
+				if(PropertyChangedEvent.Property->HasAnyPropertyFlags(CPF_Config | CPF_GlobalConfig))
+				{
+					ObjectBeingEdited->UpdateSinglePropertyInConfigFile(PropertyChangedEvent.Property, ObjectBeingEdited->GetDefaultConfigFilename());
+				}
+			}
+
+			if (bIsNewFile)
+			{
+				SettingsHelpers::CheckOutOrAddFile(FullPath);
+			}
+		}
+	}
+}
+
+template<typename T>
+class FPerCDOSettingsCustomization : public IDetailCustomization
+{
+	virtual void CustomizeDetails(IDetailLayoutBuilder& DetailBuilder) override
+	{
+		TArray<TWeakObjectPtr<UObject>> CustomizedObjects;
+		DetailBuilder.GetObjectsBeingCustomized(CustomizedObjects);
+
+		TWeakObjectPtr<UObject>& Object = CustomizedObjects[0];
+		if (Object.IsValid() && Object->HasAnyFlags(RF_ClassDefaultObject))
+		{
+			FString CategoryName = Object->GetClass()->GetName();
+			CategoryName.RemoveFromStart(T::StaticClass()->GetName());
+			CategoryName.RemoveFromEnd(TEXT("_C"));
+			IDetailCategoryBuilder& CategoryBuilder = DetailBuilder.EditCategory(Object->GetClass()->GetFName(), FText::FromString(FName::NameToDisplayString(CategoryName, false)));	// TODO: Category FName should be FullName for non-native objects
+
+			// TODO: Apply property categories as sub-categories?
+			UClass* BaseClass = Object->GetClass();
+			while (BaseClass)
+			{
+				for (FProperty* Property : TFieldRange<FProperty>(BaseClass, EFieldIteratorFlags::ExcludeSuper, EFieldIteratorFlags::ExcludeDeprecated))
+				{
+					if (Property->HasAnyPropertyFlags(CPF_Config))
+					{
+						CategoryBuilder.AddProperty(Property->GetFName(), BaseClass);
+					}
+				}
+
+				BaseClass = BaseClass != T::StaticClass() ? BaseClass->GetSuperClass() : nullptr; // Stop searching at the base type. We don't care about configurable properties lower than that.
+			}
+		}
+	}
+};
+
+
+template<typename T>
+TSharedPtr<IDetailsView> FInputEditorModule::AddClassDetailsView()
+{
+	FPropertyEditorModule& PropertyEditorModule = FModuleManager::GetModuleChecked<FPropertyEditorModule>("PropertyEditor");
+
+	FDetailsViewArgs DetailsViewArgs;
+	DetailsViewArgs.NameAreaSettings = FDetailsViewArgs::HideNameArea;
+	DetailsViewArgs.bHideSelectionTip = true;
+	DetailsViewArgs.bAllowMultipleTopLevelObjects = true;
+	DetailsViewArgs.bShowOptions = false;
+	DetailsViewArgs.bShowPropertyMatrixButton = false;
+
+	TSharedPtr<IDetailsView> DetailsView = PropertyEditorModule.CreateDetailView(DetailsViewArgs);
+	DetailsView->OnFinishedChangingProperties().AddRaw(this, &FInputEditorModule::OnSettingChanged);
+	DetailsView->RegisterInstancedCustomPropertyLayout(T::StaticClass(), FOnGetDetailCustomizationInstance::CreateLambda(&MakeShared<FPerCDOSettingsCustomization<T>>));
+
+	// Init CDOs for view
+	DetailsView->SetObjects(GatherClassDetailsCDOs(T::StaticClass(), nullptr));
+
+	return DetailsViews.Add(T::StaticClass(), DetailsView);
+}
+
+TArray<UObject*> FInputEditorModule::GatherClassDetailsCDOs(UClass* Class, const FAssetData* IgnoreAsset)
+{
+	TArray<UObject*> CDOs;
+
+	// Search native classes
+	for (TObjectIterator<UClass> ClassIt; ClassIt; ++ClassIt)
+	{
+		if (!ClassIt->IsNative() || !ClassIt->IsChildOf(Class))
+		{
+			continue;
+		}
+
+		// Ignore abstract, hidedropdown, and deprecated.
+		if (ClassIt->HasAnyClassFlags(CLASS_Abstract | CLASS_HideDropDown | CLASS_Deprecated | CLASS_NewerVersionExists))
+		{
+			continue;
+		}
+
+		CDOs.AddUnique(ClassIt->GetDefaultObject());
+	}
+
+	// Search BPs via asset registry
+	FAssetRegistryModule& AssetRegistryModule = FModuleManager::LoadModuleChecked<FAssetRegistryModule>(FName("AssetRegistry"));
+	IAssetRegistry& AssetRegistry = AssetRegistryModule.Get();
+
+	FARFilter Filter;
+	Filter.ClassNames.Add(UBlueprint::StaticClass()->GetFName());
+	Filter.bRecursiveClasses = true;
+	TArray<FAssetData> BlueprintAssetData;
+	AssetRegistry.GetAssets(Filter, BlueprintAssetData);
+
+	for (FAssetData& Asset : BlueprintAssetData)
+	{
+		if (IgnoreAsset && Asset == *IgnoreAsset)
+		{
+			continue;
+		}
+
+		FAssetDataTagMapSharedView::FFindTagResult Result = Asset.TagsAndValues.FindTag(TEXT("NativeParentClass"));
+		if (Result.IsSet())
+		{
+			const FString ClassObjectPath = FPackageName::ExportTextPathToObjectPath(Result.GetValue());
+			const FString ClassName = FPackageName::ObjectPathToObjectName(ClassObjectPath);
+			if (UClass* ParentClass = FindObjectSafe<UClass>(ANY_PACKAGE, *ClassName, true))
+			{
+				if (ParentClass->IsChildOf(Class))
+				{
+					// TODO: Forcibly loading these assets could cause problems on projects with a large number of them.
+					UBlueprint* BP = CastChecked<UBlueprint>(Asset.GetAsset());
+					CDOs.AddUnique(BP->GeneratedClass->GetDefaultObject());
+				}
+			}
+		}
+	}
+
+	// Strip objects with no config stored properties
+	CDOs.RemoveAll([Class](UObject* Object) {
+		UClass* ObjectClass = Object->GetClass();
+		if (ObjectClass->GetMetaData(TEXT("NotInputConfigurable")).ToBool())
+		{
+			return true;
+		}
+		while (ObjectClass)
+		{
+			for (FProperty* Property : TFieldRange<FProperty>(ObjectClass, EFieldIteratorFlags::ExcludeSuper, EFieldIteratorFlags::ExcludeDeprecated))
+			{
+				if (Property->HasAnyPropertyFlags(CPF_Config))
+				{
+					return false;
+				}
+			}
+
+			ObjectClass = ObjectClass != Class ? ObjectClass->GetSuperClass() : nullptr; // Stop searching at the base type. We don't care about configurable properties lower than that.
+		}
+		return true;
+	});
+
+	return CDOs;
+}
+
+void FInputEditorModule::PostEngineInit()
+{
+	// Register input settings
+	ISettingsModule* SettingsModule = FModuleManager::GetModulePtr<ISettingsModule>("Settings");
+	if (SettingsModule && FSlateApplication::IsInitialized())
+	{
+		SettingsModule->RegisterSettings("Project", "Plugins", "EnhancedInput",
+			LOCTEXT("EnhancedInputSettingsName", "Enhanced Input"),
+			LOCTEXT("EnhancedInputSettingsDescription", "Modify defaults for configurable triggers and modifiers."),
+			CreateSettingsPanel()
+		);
+	}
+}
+
+TSharedRef<SWidget> FInputEditorModule::CreateSettingsPanel()
+{
+	TSharedPtr<IDetailsView> TriggerDetailsView = AddClassDetailsView<UInputTrigger>();
+	TSharedPtr<IDetailsView> ModifierDetailsView = AddClassDetailsView<UInputModifier>();
+
+	FSlateFontInfo HeaderFont = FEditorStyle::GetFontStyle("DetailsView.CategoryFontStyle");
+	HeaderFont.Size = 32;
+
+	SAssignNew(Panel, SScrollBox)
+	+SScrollBox::Slot()
+	.Padding(0.f, 4.f, 0.f, 4.f)
+	[
+		SNew(STextBlock)
+		.Font(HeaderFont)
+		.Text(LOCTEXT("EngineInputSettingsTriggers", "Trigger Defaults"))
+	]
+
+	+SScrollBox::Slot()
+	.Padding(0.f, 4.f, 0.f, 0.f)
+	[
+		TriggerDetailsView.ToSharedRef()
+	]
+
+	+SScrollBox::Slot()
+	.Padding(0.f, 12.f, 0.f, 4.f)
+	[
+		SNew(STextBlock)
+		.Font(HeaderFont)
+		.Text(LOCTEXT("EngineInputSettingsModifiers", "Modifier Defaults"))
+	]
+
+	+SScrollBox::Slot()
+	.Padding(0.f, 8.f, 0.f, 0.f)
+	[
+		ModifierDetailsView.ToSharedRef()
+	];
+
+	return Panel.ToSharedRef();
+}
+
+void FInputEditorModule::StartupModule()
+{
+	// Register customizations
+	FPropertyEditorModule& PropertyModule = FModuleManager::GetModuleChecked<FPropertyEditorModule>("PropertyEditor");
+	PropertyModule.RegisterCustomClassLayout("InputMappingContext", FOnGetDetailCustomizationInstance::CreateStatic(&FInputContextDetails::MakeInstance));
+	PropertyModule.RegisterCustomPropertyTypeLayout("EnhancedActionKeyMapping", FOnGetPropertyTypeCustomizationInstance::CreateStatic(&FEnhancedActionMappingCustomization::MakeInstance));
+	PropertyModule.NotifyCustomizationModuleChanged();
+
+	// Register input assets
+	IAssetTools& AssetTools = FModuleManager::LoadModuleChecked<FAssetToolsModule>("AssetTools").Get();
+	InputAssetsCategory = AssetTools.RegisterAdvancedAssetCategory(FName(TEXT("Input")), LOCTEXT("InputAssetsCategory", "Input"));
+	{
+		RegisterAssetTypeActions(AssetTools, MakeShareable(new FAssetTypeActions_InputAction));
+		RegisterAssetTypeActions(AssetTools, MakeShareable(new FAssetTypeActions_InputContext));
+		RegisterAssetTypeActions(AssetTools, MakeShareable(new FAssetTypeActions_PlayerMappableInputConfig));
+		// TODO: Build these off a button on the InputContext Trigger/Mapping pickers? Would be good to have both.
+		//RegisterAssetTypeActions(AssetTools, MakeShareable(new FAssetTypeActions_InputTrigger));
+		//RegisterAssetTypeActions(AssetTools, MakeShareable(new FAssetTypeActions_InputModifier));
+	}
+
+	// Support for updating blueprint based triggers and modifiers in the settings panel
+	FAssetRegistryModule& AssetRegistryModule = FModuleManager::LoadModuleChecked<FAssetRegistryModule>(FName("AssetRegistry"));
+	IAssetRegistry& AssetRegistry = AssetRegistryModule.Get();
+	AssetRegistry.OnAssetAdded().AddRaw(this, &FInputEditorModule::OnAssetAdded);
+	AssetRegistry.OnAssetRemoved().AddRaw(this, &FInputEditorModule::OnAssetRemoved);
+	AssetRegistry.OnAssetRenamed().AddRaw(this, &FInputEditorModule::OnAssetRenamed);
+	// TODO: Update settings whenever a config variable is added/removed to an asset
+
+	FCoreDelegates::OnPostEngineInit.AddRaw(this, &FInputEditorModule::PostEngineInit);
+}
+
+void FInputEditorModule::ShutdownModule()
+{
+	// Unregister settings panel listeners
+	if (FAssetRegistryModule* AssetRegistryModule = FModuleManager::GetModulePtr<FAssetRegistryModule>("AssetRegistry"))
+	{
+		AssetRegistryModule->Get().OnAssetAdded().RemoveAll(this);
+		AssetRegistryModule->Get().OnAssetRemoved().RemoveAll(this);
+		AssetRegistryModule->Get().OnAssetRenamed().RemoveAll(this);
+	}
+
+	// Unregister input assets
+	if (FAssetToolsModule* AssetToolsModule = FModuleManager::GetModulePtr<FAssetToolsModule>("AssetTools"))
+	{
+		for (TSharedPtr<IAssetTypeActions>& AssetAction : CreatedAssetTypeActions)
+		{
+			AssetToolsModule->Get().UnregisterAssetTypeActions(AssetAction.ToSharedRef());
+		}
+	}
+	CreatedAssetTypeActions.Empty();
+
+	// Unregister input settings
+	if (ISettingsModule* SettingsModule = FModuleManager::GetModulePtr<ISettingsModule>("Settings"))
+	{
+		SettingsModule->UnregisterSettings("Project", "Engine", "Enhanced Input");
+	}
+
+	// Unregister customizations
+	FPropertyEditorModule& PropertyModule = FModuleManager::GetModuleChecked<FPropertyEditorModule>("PropertyEditor");
+	PropertyModule.UnregisterCustomClassLayout("InputContext");
+	PropertyModule.UnregisterCustomPropertyTypeLayout("EnhancedActionKeyMapping");
+	PropertyModule.NotifyCustomizationModuleChanged();
+
+	FCoreDelegates::OnPostEngineInit.RemoveAll(this);
+}
+
+FInputEditorModule::FClassDetailsView FInputEditorModule::FindClassDetailsViewForAsset(const FAssetData& AssetData)
+{
+	if (AssetData.AssetClass == UBlueprint::StaticClass()->GetFName())
+	{
+		FAssetDataTagMapSharedView::FFindTagResult Result = AssetData.TagsAndValues.FindTag(TEXT("NativeParentClass"));
+		if (Result.IsSet())
+		{
+			const FString ClassObjectPath = FPackageName::ExportTextPathToObjectPath(Result.GetValue());
+			const FString ClassName = FPackageName::ObjectPathToObjectName(ClassObjectPath);
+			if (UClass* ParentClass = FindObjectSafe<UClass>(ANY_PACKAGE, *ClassName, true))
+			{
+				if (TSharedPtr<IDetailsView>* DetailsView = DetailsViews.Find(ParentClass))
+				{
+					return FClassDetailsView(ParentClass, *DetailsView);
+				}
+			}
+		}
+	}
+
+	return FClassDetailsView();
+}
+
+void FInputEditorModule::RebuildDetailsViewForAsset(const FAssetData& AssetData, bool bIgnoreAsset)
+{
+	FClassDetailsView CDV = FindClassDetailsViewForAsset(AssetData);
+	if (CDV.IsValid())
+	{
+		CDV.View->SetObjects(GatherClassDetailsCDOs(CDV.Class, bIgnoreAsset ? &AssetData : nullptr));
+	}
+}
+
+void FInputEditorModule::Tick(float DeltaTime)
+{
+	// Update any blueprints that are referencing an input action with a modified value type
+	if (UInputAction::ActionsWithModifiedValueTypes.Num())
+	{
+		TSet<UBlueprint*> BPsModified;
+		for (TObjectIterator<UK2Node_EnhancedInputAction> NodeIt; NodeIt; ++NodeIt)
+		{
+			if (UInputAction::ActionsWithModifiedValueTypes.Contains(NodeIt->InputAction))
+			{
+				NodeIt->ReconstructNode();
+				BPsModified.Emplace(NodeIt->GetBlueprint());
+			}
+		}
+		for (TObjectIterator<UK2Node_GetInputActionValue> NodeIt; NodeIt; ++NodeIt)
+		{
+			if (UInputAction::ActionsWithModifiedValueTypes.Contains(NodeIt->InputAction))
+			{
+				NodeIt->ReconstructNode();
+				BPsModified.Emplace(NodeIt->GetBlueprint());
+			}
+		}
+
+		if (BPsModified.Num())
+		{
+			FNotificationInfo Info(FText::Format(LOCTEXT("ActionValueTypeChange", "Changing action value type affected {0} blueprint(s)!"), BPsModified.Num()));
+			Info.ExpireDuration = 5.0f;
+			FSlateNotificationManager::Get().AddNotification(Info);
+		}
+
+		UInputAction::ActionsWithModifiedValueTypes.Reset();
+	}
+}
+
+#undef LOCTEXT_NAMESPACE
\ No newline at end of file
diff --git a/Engine/Plugins/EnhancedInput/Source/InputEditor/Public/InputEditorModule.h b/Engine/Plugins/EnhancedInput/Source/InputEditor/Public/InputEditorModule.h
new file mode 100644
--- /dev/null
+++ b/Engine/Plugins/EnhancedInput/Source/InputEditor/Public/InputEditorModule.h
@@ -0,0 +1,53 @@
+// Copyright Epic Games, Inc. All Rights Reserved.
+
+#pragma once
+
+#include "CoreMinimal.h"
+#include "Factories/Factory.h"
+
+#include "InputEditorModule.generated.h"
+
+// Module is not publicly exposed
+
+// Asset factories
+
+UCLASS()
+class INPUTEDITOR_API UInputMappingContext_Factory : public UFactory
+{
+	GENERATED_UCLASS_BODY()
+public:
+	virtual UObject* FactoryCreateNew(UClass* Class, UObject* InParent, FName Name, EObjectFlags Flags, UObject* Context, FFeedbackContext* Warn) override;
+};
+
+UCLASS()
+class INPUTEDITOR_API UInputAction_Factory : public UFactory
+{
+	GENERATED_UCLASS_BODY()
+public:
+	virtual UObject* FactoryCreateNew(UClass* Class, UObject* InParent, FName Name, EObjectFlags Flags, UObject* Context, FFeedbackContext* Warn) override;
+};
+
+UCLASS()
+class INPUTEDITOR_API UPlayerMappableInputConfig_Factory : public UFactory
+{
+	GENERATED_UCLASS_BODY()
+public:
+	virtual UObject* FactoryCreateNew(UClass* Class, UObject* InParent, FName Name, EObjectFlags Flags, UObject* Context, FFeedbackContext* Warn) override;
+};
+
+// TODO: Add trigger/modifier factories and hook up RegisterAssetTypeActions type construction.
+//
+//UCLASS()
+//class INPUTEDITOR_API UInputTrigger_Factory : public UBlueprintFactory
+//{
+//	GENERATED_UCLASS_BODY()
+//};
+//
+//UCLASS()
+//class INPUTEDITOR_API UInputModifier_Factory : public UBlueprintFactory
+//{
+//	GENERATED_UCLASS_BODY()
+//
+//	UPROPERTY(EditAnywhere, Category = DataAsset)
+//	TSubclassOf<UDataAsset> DataAssetClass;
+//};
