Enabled Enhanced Input plugin by default

diff --git a/Engine/Plugins/EnhancedInput/EnhancedInput.uplugin b/Engine/Plugins/EnhancedInput/EnhancedInput.uplugin
--- a/Engine/Plugins/EnhancedInput/EnhancedInput.uplugin
+++ b/Engine/Plugins/EnhancedInput/EnhancedInput.uplugin
@@ -10,7 +10,7 @@
 	"DocsURL": "",
 	"MarketplaceURL": "",
 	"SupportURL": "",
-	"EnabledByDefault": false,
+	"EnabledByDefault": true,
 	"CanContainContent": false,
 	"IsBetaVersion": true,
 	"Installed": false,
