Backout 19284661

#rnx

diff --git a/Engine/Plugins/EnhancedInput/EnhancedInput.uplugin b/Engine/Plugins/EnhancedInput/EnhancedInput.uplugin
--- a/Engine/Plugins/EnhancedInput/EnhancedInput.uplugin
+++ b/Engine/Plugins/EnhancedInput/EnhancedInput.uplugin
@@ -10,7 +10,7 @@
 	"DocsURL": "",
 	"MarketplaceURL": "",
 	"SupportURL": "",
-	"EnabledByDefault": true,
+	"EnabledByDefault": false,
 	"CanContainContent": false,
 	"IsBetaVersion": true,
 	"Installed": false,
