Remove the UEnhancedInputEngineSubsystem. This will be replaced in favor of a "World" subsystem after the player input has been separated from the player controller. This was never useable and was not a blueprint type.

diff --git a/Engine/Plugins/EnhancedInput/Source/EnhancedInput/Private/EnhancedInputLibrary.cpp b/Engine/Plugins/EnhancedInput/Source/EnhancedInput/Private/EnhancedInputLibrary.cpp
--- a/Engine/Plugins/EnhancedInput/Source/EnhancedInput/Private/EnhancedInputLibrary.cpp
+++ b/Engine/Plugins/EnhancedInput/Source/EnhancedInput/Private/EnhancedInputLibrary.cpp
@@ -15,8 +15,8 @@
 
 void UEnhancedInputLibrary::ForEachSubsystem(TFunctionRef<void(IEnhancedInputSubsystemInterface*)> SubsystemPredicate)
 {
-	// Engine
-	SubsystemPredicate(Cast<IEnhancedInputSubsystemInterface>(GEngine->GetEngineSubsystem<UEnhancedInputEngineSubsystem>()));
+	// TODO: World subsystem for enhanced input, so that you bind to actions without an owning player controller
+	// This is useful for widgets, main menu situations, or just stuff that isn't dependant on a player.
 
 	// Players
 	for (TObjectIterator<UEnhancedInputLocalPlayerSubsystem> It; It; ++It)
diff --git a/Engine/Plugins/EnhancedInput/Source/EnhancedInput/Private/EnhancedInputSubsystems.cpp b/Engine/Plugins/EnhancedInput/Source/EnhancedInput/Private/EnhancedInputSubsystems.cpp
--- a/Engine/Plugins/EnhancedInput/Source/EnhancedInput/Private/EnhancedInputSubsystems.cpp
+++ b/Engine/Plugins/EnhancedInput/Source/EnhancedInput/Private/EnhancedInputSubsystems.cpp
@@ -27,21 +27,4 @@ UEnhancedPlayerInput* UEnhancedInputLocalPlayerSubsystem::GetPlayerInput() const
 		return Cast<UEnhancedPlayerInput>(PlayerController->PlayerInput);
 	}
 	return nullptr;
-}
-
-// **************************************************************************************************
-// *
-// * UEnhancedInputEngineSubsystem
-// *
-// **************************************************************************************************
-
-void UEnhancedInputEngineSubsystem::Initialize(FSubsystemCollectionBase& Collection)
-{	
-	PlayerInput = nullptr;
-	//PlayerInput = NewObject<UEnhancedPlayerInput>();	// TODO: Remove UPlayerInput Within? Create an empty player controller?
-}
-
-UEnhancedPlayerInput* UEnhancedInputEngineSubsystem::GetPlayerInput() const
-{
-	return PlayerInput;
-}
+}
\ No newline at end of file
diff --git a/Engine/Plugins/EnhancedInput/Source/EnhancedInput/Public/EnhancedInputSubsystems.h b/Engine/Plugins/EnhancedInput/Source/EnhancedInput/Public/EnhancedInputSubsystems.h
--- a/Engine/Plugins/EnhancedInput/Source/EnhancedInput/Public/EnhancedInputSubsystems.h
+++ b/Engine/Plugins/EnhancedInput/Source/EnhancedInput/Public/EnhancedInputSubsystems.h
@@ -21,26 +21,3 @@ public:
 	virtual UEnhancedPlayerInput* GetPlayerInput() const override;
 	// End IEnhancedInputSubsystemInterface
 };
-
-// Global input handling subsystem
-// TODO: For now this is a non-functional placeholder that shuold be enabled when the Enhanced Input Plugin is enabled by default.
-UCLASS(NotBlueprintType)
-class ENHANCEDINPUT_API UEnhancedInputEngineSubsystem : public UEngineSubsystem, public IEnhancedInputSubsystemInterface
-{
-	GENERATED_BODY()
-
-private:
-	UPROPERTY()
-	UEnhancedPlayerInput* PlayerInput;
-
-public:
-
-	// Begin USubsystem
-	virtual void Initialize(FSubsystemCollectionBase& Collection) override;
-	// End USubsystem
-
-	// Begin IEnhancedInputSubsystemInterface
-	virtual UEnhancedPlayerInput* GetPlayerInput() const override;
-	// End IEnhancedInputSubsystemInterface
-
-};
