Remove the unusable "Collection" input modifier. There is no use deprecating this because it was unusable. If you had it in one of your assets then it was not working in the first place and should throw an error. This also doesn't need deprecation because enhanced input is coming out of experimental, so we have the rights to remove stuff like this. No point in keeping the code rot around right from the beginning.

diff --git a/Engine/Plugins/EnhancedInput/Source/EnhancedInput/Private/InputModifiers.cpp b/Engine/Plugins/EnhancedInput/Source/EnhancedInput/Private/InputModifiers.cpp
--- a/Engine/Plugins/EnhancedInput/Source/EnhancedInput/Private/InputModifiers.cpp
+++ b/Engine/Plugins/EnhancedInput/Source/EnhancedInput/Private/InputModifiers.cpp
@@ -329,21 +329,4 @@ FLinearColor UInputModifierSwizzleAxis::GetVisualizationColor_Implementation(FIn
 	// TODO: Color blend per swizzle type?
 	float SampleX = (FMath::Abs(SampleValue.Get<float>()) + 1.f) * 0.5f;
 	return FLinearColor(SampleX, 1.f - SampleX, 0.f);
-}
-
-/*
-* Modifier collections
-*/
-FInputActionValue UInputModifierCollection::ModifyRaw_Implementation(const UEnhancedPlayerInput* PlayerInput, FInputActionValue CurrentValue, float DeltaTime)
-{
-	FInputActionValue ModifiedValue = CurrentValue;
-	for (UInputModifier* Modifier : Modifiers)
-	{
-		ModifiedValue = Modifier->ModifyRaw(PlayerInput, ModifiedValue, DeltaTime);
-		if (!bPermitValueTypeModification)
-		{
-			ModifiedValue.ConvertToType(CurrentValue);
-		}
-	}
-	return ModifiedValue;
 }
\ No newline at end of file
diff --git a/Engine/Plugins/EnhancedInput/Source/EnhancedInput/Public/InputModifiers.h b/Engine/Plugins/EnhancedInput/Source/EnhancedInput/Public/InputModifiers.h
--- a/Engine/Plugins/EnhancedInput/Source/EnhancedInput/Public/InputModifiers.h
+++ b/Engine/Plugins/EnhancedInput/Source/EnhancedInput/Public/InputModifiers.h
@@ -306,29 +306,4 @@ public:
 protected:
 	virtual FInputActionValue ModifyRaw_Implementation(const UEnhancedPlayerInput* PlayerInput, FInputActionValue CurrentValue, float DeltaTime) override;
 	virtual FLinearColor GetVisualizationColor_Implementation(FInputActionValue SampleValue, FInputActionValue FinalValue) const override;
-};
-
-/** Modifier collection
-	* A user definable group of modifiers that can be easily applied to multiple actions or mappings to save duplication work.
-	*/
-UCLASS(NotBlueprintable, MinimalAPI, meta = (DisplayName = "Modifier Collection"))
-class UInputModifierCollection : public	UInputModifier
-{
-	GENERATED_BODY()
-
-public:
-
-	UPROPERTY(EditInstanceOnly, BlueprintReadWrite, Category = Settings)
-	TArray<UInputModifier*> Modifiers;	// TODO: These either need to be instanced per copy of the collection or it needs to be clear the modifiers are okay running multiple times per frame (or we just don't support collections)
-
-	/** If set each modifier will not have the modified value corrected to the base type before execution.
-	*	After all modifiers are run the resulting value will be converted back to the action's value type as with any other modifier.
-	*	This allows for complex sets of conditional modifiers that can alter their behavior based on their predecessors value type.
-	*	Note that this is an advanced feature and may cause issues if used with the basic modifier implementations.
-	**/
-	UPROPERTY(EditAnywhere, BlueprintReadWrite, Category = Settings)
-	bool bPermitValueTypeModification = false;
-
-protected:
-	virtual FInputActionValue ModifyRaw_Implementation(const UEnhancedPlayerInput* PlayerInput, FInputActionValue CurrentValue, float DeltaTime) override;
 };
\ No newline at end of file
diff --git a/Engine/Plugins/EnhancedInput/Source/InputEditor/Private/Tests/InputModifierTest.cpp b/Engine/Plugins/EnhancedInput/Source/InputEditor/Private/Tests/InputModifierTest.cpp
--- a/Engine/Plugins/EnhancedInput/Source/InputEditor/Private/Tests/InputModifierTest.cpp
+++ b/Engine/Plugins/EnhancedInput/Source/InputEditor/Private/Tests/InputModifierTest.cpp
@@ -205,28 +205,3 @@ bool FInputModifierDeadzoneTest::RunTest(const FString& Parameters)
 
 	return true;
 }
-
-IMPLEMENT_SIMPLE_AUTOMATION_TEST(FInputModifierCollectionTest, "Input.Modifiers.Collections", BasicModifierTestFlags)
-
-bool FInputModifierCollectionTest::RunTest(const FString& Parameters)
-{
-	GIVEN(UControllablePlayer & Data = ABasicModifierTest(this, EInputActionValueType::Axis1D));
-	AND(AnActionIsMappedToAKey(Data, TestContext, TestAction, TestKey));
-
-	// Build collection modifier
-	UInputModifierCollection* Collection = NewObject<UInputModifierCollection>();
-	UInputModifierNegate* Negate = NewObject<UInputModifierNegate>();
-	UInputModifierScalar* Scalar = NewObject<UInputModifierScalar>(); Scalar->Scalar = FVector(0.5f, 0.5f, 0.5f);
-	Collection->Modifiers.Add(Negate);
-	Collection->Modifiers.Add(Scalar);
-
-	WHEN(AModifierIsAppliedToAnActionMapping(Data, Collection, TestContext, TestAction, TestKey));
-	AND(AKeyIsActuated(Data, TestKey));
-	AND(InputIsTicked(Data));
-
-	THEN(PressingKeyTriggersAction(Data, TestAction));
-	AND(TestEqual(TEXT("Tnput value"), GetActionValue(Data, TestAction), -0.5f));
-
-	return true;
-}
-
