Fix non-unity compile error

#rnx

diff --git a/Engine/Plugins/EnhancedInput/Source/EnhancedInput/Private/EnhancedPlayerInput.cpp b/Engine/Plugins/EnhancedInput/Source/EnhancedInput/Private/EnhancedPlayerInput.cpp
--- a/Engine/Plugins/EnhancedInput/Source/EnhancedInput/Private/EnhancedPlayerInput.cpp
+++ b/Engine/Plugins/EnhancedInput/Source/EnhancedInput/Private/EnhancedPlayerInput.cpp
@@ -7,6 +7,8 @@
 #include "InputModifiers.h"
 #include "InputTriggers.h"
 #include "Engine/World.h"
+#include "GameFramework/WorldSettings.h"
+
 // NOTE: Enum order represents firing priority(lowest to highest) and is important as multiple keys bound to the same action may generate differing trigger event states.
 enum class ETriggerEventInternal : uint8
 {
