Only call FImageUtils::CreateTexture2D if we are running WITH_EDITOR when drawing enhanced input debug info. FImageUtils::CreateTexture2D will throw a fatal error if you are running without WITH_EDITOR (on console, or any packaged game). Also account for that texture being null then.

Only bind to the AHUD::OnShowDebugInfo delegate if ENABLE_DRAW_DEBUG is on, otherwise it is pointless to go through any debug drawing logic.

diff --git a/Engine/Plugins/EnhancedInput/Source/EnhancedInput/Private/EnhancedInputModule.cpp b/Engine/Plugins/EnhancedInput/Source/EnhancedInput/Private/EnhancedInputModule.cpp
--- a/Engine/Plugins/EnhancedInput/Source/EnhancedInput/Private/EnhancedInputModule.cpp
+++ b/Engine/Plugins/EnhancedInput/Source/EnhancedInput/Private/EnhancedInputModule.cpp
@@ -232,8 +232,9 @@ void FEnhancedInputModule::StartupModule()
 
 	if (!IsRunningDedicatedServer())
 	{
+#if ENABLE_DRAW_DEBUG
 		AHUD::OnShowDebugInfo.AddStatic(&FEnhancedInputModule::OnShowDebugInfo);
-
+#endif
 	    // Register console commands
 	    ConsoleCommands.Add(IConsoleManager::Get().RegisterConsoleCommand(
 		    TEXT("Input.+action"),
diff --git a/Engine/Plugins/EnhancedInput/Source/EnhancedInput/Private/EnhancedInputSubsystemInterfaceDebug.cpp b/Engine/Plugins/EnhancedInput/Source/EnhancedInput/Private/EnhancedInputSubsystemInterfaceDebug.cpp
--- a/Engine/Plugins/EnhancedInput/Source/EnhancedInput/Private/EnhancedInputSubsystemInterfaceDebug.cpp
+++ b/Engine/Plugins/EnhancedInput/Source/EnhancedInput/Private/EnhancedInputSubsystemInterfaceDebug.cpp
@@ -25,7 +25,10 @@ struct FVisualizationTexture
 	// Move only
 	FVisualizationTexture(UTexture2D* InTexture) : Texture(InTexture)
 	{
-		Texture->AddToRoot();
+		if (Texture)
+		{
+			Texture->AddToRoot();
+		}
 	}
 	FVisualizationTexture(FVisualizationTexture&& Other)
 	{
@@ -357,7 +360,12 @@ void IEnhancedInputSubsystemInterface::ShowDebugActionModifiers(UCanvas* Canvas,
 			FCreateTexture2DParameters Params;
 			Params.bDeferCompression = true;
 			Params.CompressionSettings = TC_EditorIcon;
+			// FImageUtils::CreateTexture2D will return null and throw a fatal error unless we have WITH_EDITOR
+#if WITH_EDITOR
 			return FImageUtils::CreateTexture2D(VisSize.X, VisSize.Y, ColorData, IEnhancedInputModule::Get().GetLibrary(), FString(), EObjectFlags::RF_NoFlags, Params);
+#else 
+			return nullptr;
+#endif
 		};
 
 		bool bDrawnAtLeastOneModifier = false;
