First pass of making Enhanced Input use TObjectPtr instead of raw pointers.

#rnx

diff --git a/Engine/Plugins/EnhancedInput/Source/EnhancedInput/Public/InputModifiers.h b/Engine/Plugins/EnhancedInput/Source/EnhancedInput/Public/InputModifiers.h
--- a/Engine/Plugins/EnhancedInput/Source/EnhancedInput/Public/InputModifiers.h
+++ b/Engine/Plugins/EnhancedInput/Source/EnhancedInput/Public/InputModifiers.h
@@ -201,6 +201,7 @@ protected:
 	virtual FInputActionValue ModifyRaw_Implementation(const UEnhancedPlayerInput* PlayerInput, FInputActionValue CurrentValue, float DeltaTime) override;
 };
 
+class UCurveFloat;
 
 /** Response Curve User Defined
 	*  Apply a custom response curve to input values, per axis
@@ -213,11 +214,11 @@ class UInputModifierResponseCurveUser : public	UInputModifier
 public:
 
 	UPROPERTY(EditInstanceOnly, BlueprintReadWrite, Category = Settings, meta = (DisplayThumbnail = "false"))
-	class UCurveFloat* ResponseX;
+	TObjectPtr<UCurveFloat> ResponseX;
 	UPROPERTY(EditInstanceOnly, BlueprintReadWrite, Category = Settings, meta = (DisplayThumbnail = "false"))
-	class UCurveFloat* ResponseY;
+	TObjectPtr<UCurveFloat> ResponseY;
 	UPROPERTY(EditInstanceOnly, BlueprintReadWrite, Category = Settings, meta = (DisplayThumbnail = "false"))
-	class UCurveFloat* ResponseZ;
+	TObjectPtr<UCurveFloat> ResponseZ;
 
 protected:
 	virtual FInputActionValue ModifyRaw_Implementation(const UEnhancedPlayerInput* PlayerInput, FInputActionValue CurrentValue, float DeltaTime) override;
diff --git a/Engine/Plugins/EnhancedInput/Source/EnhancedInput/Public/InputTriggers.h b/Engine/Plugins/EnhancedInput/Source/EnhancedInput/Public/InputTriggers.h
--- a/Engine/Plugins/EnhancedInput/Source/EnhancedInput/Public/InputTriggers.h
+++ b/Engine/Plugins/EnhancedInput/Source/EnhancedInput/Public/InputTriggers.h
@@ -387,7 +387,7 @@ public:
 
 	// The action that must be triggering for this trigger's action to trigger
 	UPROPERTY(EditInstanceOnly, BlueprintReadWrite, Category = "Trigger Settings", meta = (DisplayThumbnail = "false"))
-	const class UInputAction* ChordAction = nullptr;
+	TObjectPtr<const class UInputAction> ChordAction = nullptr;
 };
 
 /** UInputTriggerChordBlocker
diff --git a/Engine/Plugins/EnhancedInput/Source/InputBlueprintNodes/Private/K2Node_EnhancedInputActionEvent.h b/Engine/Plugins/EnhancedInput/Source/InputBlueprintNodes/Private/K2Node_EnhancedInputActionEvent.h
--- a/Engine/Plugins/EnhancedInput/Source/InputBlueprintNodes/Private/K2Node_EnhancedInputActionEvent.h
+++ b/Engine/Plugins/EnhancedInput/Source/InputBlueprintNodes/Private/K2Node_EnhancedInputActionEvent.h
@@ -17,7 +17,7 @@ class INPUTBLUEPRINTNODES_API UK2Node_EnhancedInputActionEvent : public UK2Node_
 	GENERATED_UCLASS_BODY()
 
 	UPROPERTY()
-	const UInputAction* InputAction;
+	TObjectPtr<const UInputAction> InputAction;
 
 	UPROPERTY()
 	ETriggerEvent TriggerEvent;
diff --git a/Engine/Plugins/EnhancedInput/Source/InputBlueprintNodes/Private/K2Node_InputActionValueAccessor.cpp b/Engine/Plugins/EnhancedInput/Source/InputBlueprintNodes/Private/K2Node_InputActionValueAccessor.cpp
--- a/Engine/Plugins/EnhancedInput/Source/InputBlueprintNodes/Private/K2Node_InputActionValueAccessor.cpp
+++ b/Engine/Plugins/EnhancedInput/Source/InputBlueprintNodes/Private/K2Node_InputActionValueAccessor.cpp
@@ -21,7 +21,7 @@ void UK2Node_InputActionValueAccessor::AllocateDefaultPins()
 	Super::AllocateDefaultPins();
 
 	UEdGraphPin* InputActionPin = FindPinChecked(TEXT("Action"));
-	InputActionPin->DefaultObject = const_cast<UInputAction*>(InputAction);
+	InputActionPin->DefaultObject = const_cast<UInputAction*>(ToRawPtr(InputAction));
 }
 
 UClass* UK2Node_InputActionValueAccessor::GetDynamicBindingClass() const
diff --git a/Engine/Plugins/EnhancedInput/Source/InputBlueprintNodes/Private/K2Node_InputActionValueAccessor.h b/Engine/Plugins/EnhancedInput/Source/InputBlueprintNodes/Private/K2Node_InputActionValueAccessor.h
--- a/Engine/Plugins/EnhancedInput/Source/InputBlueprintNodes/Private/K2Node_InputActionValueAccessor.h
+++ b/Engine/Plugins/EnhancedInput/Source/InputBlueprintNodes/Private/K2Node_InputActionValueAccessor.h
@@ -29,5 +29,5 @@ public:
 
 private:
 	UPROPERTY()
-	const UInputAction* InputAction;
+	TObjectPtr<const UInputAction> InputAction;
 };
diff --git a/Engine/Plugins/EnhancedInput/Source/InputBlueprintNodes/Public/K2Node_EnhancedInputAction.h b/Engine/Plugins/EnhancedInput/Source/InputBlueprintNodes/Public/K2Node_EnhancedInputAction.h
--- a/Engine/Plugins/EnhancedInput/Source/InputBlueprintNodes/Public/K2Node_EnhancedInputAction.h
+++ b/Engine/Plugins/EnhancedInput/Source/InputBlueprintNodes/Public/K2Node_EnhancedInputAction.h
@@ -23,7 +23,7 @@ class INPUTBLUEPRINTNODES_API UK2Node_EnhancedInputAction : public UK2Node, publ
 	GENERATED_UCLASS_BODY()
 
 	UPROPERTY()
-	const UInputAction* InputAction;
+	TObjectPtr<const UInputAction> InputAction;
 
 	//~ Begin UEdGraphNode Interface.
 	virtual void AllocateDefaultPins() override;
diff --git a/Engine/Plugins/EnhancedInput/Source/InputBlueprintNodes/Public/K2Node_GetInputActionValue.h b/Engine/Plugins/EnhancedInput/Source/InputBlueprintNodes/Public/K2Node_GetInputActionValue.h
--- a/Engine/Plugins/EnhancedInput/Source/InputBlueprintNodes/Public/K2Node_GetInputActionValue.h
+++ b/Engine/Plugins/EnhancedInput/Source/InputBlueprintNodes/Public/K2Node_GetInputActionValue.h
@@ -20,7 +20,7 @@ class UK2Node_GetInputActionValue : public UK2Node
 	GENERATED_UCLASS_BODY()
 
 	UPROPERTY()
-	const UInputAction* InputAction;
+	TObjectPtr<const UInputAction> InputAction;
 
 	//~ Begin EdGraphNode Interface
 	virtual void AllocateDefaultPins() override;
