More TObjectPtr replacement for enhanced input

#rnx

diff --git a/Engine/Plugins/EnhancedInput/Source/EnhancedInput/Private/EnhancedInputActionDelegateBinding.cpp b/Engine/Plugins/EnhancedInput/Source/EnhancedInput/Private/EnhancedInputActionDelegateBinding.cpp
--- a/Engine/Plugins/EnhancedInput/Source/EnhancedInput/Private/EnhancedInputActionDelegateBinding.cpp
+++ b/Engine/Plugins/EnhancedInput/Source/EnhancedInput/Private/EnhancedInputActionDelegateBinding.cpp
@@ -19,7 +19,7 @@ void UEnhancedInputActionDelegateBinding::BindToInputComponent(UInputComponent*
 
 	for(const FBlueprintEnhancedInputActionBinding& Binding : InputActionDelegateBindings)
 	{
-		Component->BindAction(Binding.InputAction, Binding.TriggerEvent, ObjectToBindTo, Binding.FunctionNameToBind);
+		Component->BindAction(Binding.InputAction.Get(), Binding.TriggerEvent, ObjectToBindTo, Binding.FunctionNameToBind);
 	}
 }
 
diff --git a/Engine/Plugins/EnhancedInput/Source/EnhancedInput/Private/EnhancedInputSubsystemInterface.cpp b/Engine/Plugins/EnhancedInput/Source/EnhancedInput/Private/EnhancedInputSubsystemInterface.cpp
--- a/Engine/Plugins/EnhancedInput/Source/EnhancedInput/Private/EnhancedInputSubsystemInterface.cpp
+++ b/Engine/Plugins/EnhancedInput/Source/EnhancedInput/Private/EnhancedInputSubsystemInterface.cpp
@@ -122,14 +122,14 @@ EMappingQueryResult IEnhancedInputSubsystemInterface::QueryMapKeyInActiveContext
 	}
 
 	// TODO: Inefficient, but somewhat forgivable as the mapping context count is likely to be single figure.
-	TMap<const UInputMappingContext*, int32> OrderedInputContexts = PlayerInput->AppliedInputContexts;
+	TMap<TObjectPtr<const UInputMappingContext>, int32> OrderedInputContexts = PlayerInput->AppliedInputContexts;
 	OrderedInputContexts.ValueSort([](const int32& A, const int32& B) { return A > B; });
 
 	TArray<UInputMappingContext*> Applied;
 	Applied.Reserve(OrderedInputContexts.Num());
-	for (const TPair<const UInputMappingContext*, int32>& ContextPair : OrderedInputContexts)
+	for (const TPair<TObjectPtr<const UInputMappingContext>, int32>& ContextPair : OrderedInputContexts)
 	{
-		Applied.Add(const_cast<UInputMappingContext*>(ContextPair.Key));
+		Applied.Add(const_cast<UInputMappingContext*>(ToRawPtr(ContextPair.Key)));
 	}
 
 	return QueryMapKeyInContextSet(Applied, InputContext, Action, Key, OutIssues, BlockingIssues);
@@ -420,7 +420,7 @@ void IEnhancedInputSubsystemInterface::AddPlayerMappableConfig(const UPlayerMapp
 {
 	if(Config)
 	{
-		for(TPair<UInputMappingContext*, int32> Pair : Config->GetMappingContexts())
+		for(TPair<TObjectPtr<UInputMappingContext>, int32> Pair : Config->GetMappingContexts())
 		{
 			AddMappingContext(Pair.Key, Pair.Value, Options);
 		}	
@@ -431,7 +431,7 @@ void IEnhancedInputSubsystemInterface::RemovePlayerMappableConfig(const UPlayerM
 {
 	if(Config)
 	{
-		for(TPair<UInputMappingContext*, int32> Pair : Config->GetMappingContexts())
+		for(TPair<TObjectPtr<UInputMappingContext>, int32> Pair : Config->GetMappingContexts())
 		{
 			RemoveMappingContext(Pair.Key, Options);
 		}	
@@ -548,14 +548,14 @@ void IEnhancedInputSubsystemInterface::RebuildControlMappings()
 	PlayerInput->ClearAllMappings();
 
 	// Order contexts by priority
-	TMap<const UInputMappingContext*, int32> OrderedInputContexts = PlayerInput->AppliedInputContexts;
+	TMap<TObjectPtr<const UInputMappingContext>, int32> OrderedInputContexts = PlayerInput->AppliedInputContexts;
 	OrderedInputContexts.ValueSort([](const int32& A, const int32& B) { return A > B; });
 
 	TSet<FKey> AppliedKeys;
 
 	TArray<int32> ChordedMappings;
 
-	for (const TPair<const UInputMappingContext*, int32>& ContextPair : OrderedInputContexts)
+	for (const TPair<TObjectPtr<const UInputMappingContext>, int32>& ContextPair : OrderedInputContexts)
 	{
 		// Don't apply context specific keys immediately, allowing multiple mappings to the same key within the same context if required.
 		TArray<FKey> ContextAppliedKeys;
@@ -633,9 +633,9 @@ void IEnhancedInputSubsystemInterface::RebuildControlMappings()
 		
 		// Remove action instance data for actions that are not referenced in the new action mappings
 		TSet<const UInputAction*> RemovedActions;
-		for (TPair<const UInputAction*, FInputActionInstance>& ActionInstance : PlayerInput->ActionInstanceData)
+		for (TPair<TObjectPtr<const UInputAction>, FInputActionInstance>& ActionInstance : PlayerInput->ActionInstanceData)
 		{
-			RemovedActions.Add(ActionInstance.Key);
+			RemovedActions.Add(ActionInstance.Key.Get());
 		}
 
 		// Return true if the given FKey was in the old Player Input mappings
diff --git a/Engine/Plugins/EnhancedInput/Source/EnhancedInput/Private/EnhancedInputSubsystemInterfaceDebug.cpp b/Engine/Plugins/EnhancedInput/Source/EnhancedInput/Private/EnhancedInputSubsystemInterfaceDebug.cpp
--- a/Engine/Plugins/EnhancedInput/Source/EnhancedInput/Private/EnhancedInputSubsystemInterfaceDebug.cpp
+++ b/Engine/Plugins/EnhancedInput/Source/EnhancedInput/Private/EnhancedInputSubsystemInterfaceDebug.cpp
@@ -123,13 +123,13 @@ void IEnhancedInputSubsystemInterface::ShowDebugInfo(UCanvas* Canvas)
 			}
 		};
 
-		TMap<const UInputMappingContext*, int32> OrderedInputContexts = PlayerInput->AppliedInputContexts;
+		TMap<TObjectPtr<const UInputMappingContext>, int32> OrderedInputContexts = PlayerInput->AppliedInputContexts;
 		OrderedInputContexts.ValueSort([](const int32& A, const int32& B) { return A > B; });
 
 		// Work through all input contexts, displaying active mappings, overridden mappings, etc.
-		for (const TPair<const UInputMappingContext*, int32>& ContextPair : OrderedInputContexts)
+		for (const TPair<TObjectPtr<const UInputMappingContext>, int32>& ContextPair : OrderedInputContexts)
 		{
-			const UInputMappingContext* Context = ContextPair.Key;
+			const UInputMappingContext* Context = ContextPair.Key.Get();
 			DisplayDebugManager.SetDrawColor(FColor::Yellow);
 			DisplayDebugManager.DrawString(FString::Printf(TEXT("  Context: %s"), *Context->GetFName().ToString()));
 
diff --git a/Engine/Plugins/EnhancedInput/Source/EnhancedInput/Private/EnhancedPlayerInput.cpp b/Engine/Plugins/EnhancedInput/Source/EnhancedInput/Private/EnhancedPlayerInput.cpp
--- a/Engine/Plugins/EnhancedInput/Source/EnhancedInput/Private/EnhancedPlayerInput.cpp
+++ b/Engine/Plugins/EnhancedInput/Source/EnhancedInput/Private/EnhancedPlayerInput.cpp
@@ -331,7 +331,7 @@ void UEnhancedPlayerInput::ProcessInputStack(const TArray<UInputComponent*>& Inp
 
 
 	// Post tick action instance updates
-	for (TPair<const UInputAction*, FInputActionInstance>& ActionPair : ActionInstanceData)
+	for (TPair<TObjectPtr<const UInputAction>, FInputActionInstance>& ActionPair : ActionInstanceData)
 	{
 		TRACE_CPUPROFILER_EVENT_SCOPE(EnhPIS_PostTick);
 
@@ -505,7 +505,7 @@ void UEnhancedPlayerInput::ProcessInputStack(const TArray<UInputComponent*>& Inp
 	}
 
 	// Reset action instance timers where necessary post delegate calls
-	for (TPair<const UInputAction*, FInputActionInstance>& ActionPair : ActionInstanceData)
+	for (TPair<TObjectPtr<const UInputAction>, FInputActionInstance>& ActionPair : ActionInstanceData)
 	{
 		FInputActionInstance& ActionData = ActionPair.Value;
 		switch (ActionData.TriggerEvent)
diff --git a/Engine/Plugins/EnhancedInput/Source/EnhancedInput/Private/PlayerMappableInputConfig.cpp b/Engine/Plugins/EnhancedInput/Source/EnhancedInput/Private/PlayerMappableInputConfig.cpp
--- a/Engine/Plugins/EnhancedInput/Source/EnhancedInput/Private/PlayerMappableInputConfig.cpp
+++ b/Engine/Plugins/EnhancedInput/Source/EnhancedInput/Private/PlayerMappableInputConfig.cpp
@@ -30,12 +30,10 @@ EDataValidationResult UPlayerMappableInputConfig::IsDataValid(TArray<FText>& Val
 		Result = CombineDataValidationResults(Result, EDataValidationResult::Invalid);
 	}
 	
-	//TMap<const UInputMappingContext*, int32> InputContexts = Contexts;
-	
 	// Check that every context is valid and all player bindable mappings have names
-	for (const TPair<UInputMappingContext*, int32>& ContextPair : Contexts)
+	for (const TPair<TObjectPtr<UInputMappingContext>, int32>& ContextPair : Contexts)
 	{
-		if(const UInputMappingContext* IMC = ContextPair.Key)
+		if(const UInputMappingContext* IMC = ContextPair.Key.Get())
 		{
 			for(const FEnhancedActionKeyMapping& Mapping : IMC->GetMappings())
 			{
@@ -70,9 +68,9 @@ EDataValidationResult UPlayerMappableInputConfig::IsDataValid(TArray<FText>& Val
 
 void UPlayerMappableInputConfig::ForEachDefaultPlayerMappableKey(TFunctionRef<void(const FEnhancedActionKeyMapping&)> Operation) const
 {
-	for (const TPair<UInputMappingContext*, int32>& ContextPair: Contexts)
+	for (const TPair<TObjectPtr<UInputMappingContext>, int32>& ContextPair: Contexts)
 	{
-		if(const UInputMappingContext* IMC = ContextPair.Key)
+		if(const UInputMappingContext* IMC = ContextPair.Key.Get())
 		{
 			for(const FEnhancedActionKeyMapping& Mapping : IMC->GetMappings())
 			{
diff --git a/Engine/Plugins/EnhancedInput/Source/EnhancedInput/Public/EnhancedActionKeyMapping.h b/Engine/Plugins/EnhancedInput/Source/EnhancedInput/Public/EnhancedActionKeyMapping.h
--- a/Engine/Plugins/EnhancedInput/Source/EnhancedInput/Public/EnhancedActionKeyMapping.h
+++ b/Engine/Plugins/EnhancedInput/Source/EnhancedInput/Public/EnhancedActionKeyMapping.h
@@ -49,6 +49,9 @@ public:
 	FText DisplayCategory = FText::GetEmpty();
 };
 
+class UInputTrigger;
+class UInputModifier;
+
 /**
  * Defines a mapping between a key activation and the resulting enhanced action
  * An key could be a button press, joystick axis movement, etc.
@@ -62,7 +65,7 @@ struct FEnhancedActionKeyMapping
 
 	/** Action to be affected by the key  */
 	UPROPERTY(EditAnywhere, BlueprintReadOnly, Category = "Input")
-	const UInputAction* Action = nullptr;
+	TObjectPtr<const UInputAction> Action = nullptr;
 
 	/** Key that affect the action. */
 	UPROPERTY(EditAnywhere, BlueprintReadOnly, Category = "Input")
@@ -93,14 +96,14 @@ struct FEnhancedActionKeyMapping
 	* All Implicit triggers in this list must be met.
 	*/
 	UPROPERTY(EditAnywhere, Instanced, BlueprintReadWrite, Category = "Input")
-	TArray<class UInputTrigger*> Triggers;
+	TArray<TObjectPtr<UInputTrigger>> Triggers;
 
 	/**
 	* Modifiers applied to the raw key value.
 	* These are applied sequentially in array order.
 	*/
 	UPROPERTY(EditAnywhere, Instanced, BlueprintReadWrite, Category = "Input")
-	TArray<class UInputModifier*> Modifiers;
+	TArray<TObjectPtr<UInputModifier>> Modifiers;
 
 	bool operator==(const FEnhancedActionKeyMapping& Other) const
 	{
diff --git a/Engine/Plugins/EnhancedInput/Source/EnhancedInput/Public/EnhancedInputActionDelegateBinding.h b/Engine/Plugins/EnhancedInput/Source/EnhancedInput/Public/EnhancedInputActionDelegateBinding.h
--- a/Engine/Plugins/EnhancedInput/Source/EnhancedInput/Public/EnhancedInputActionDelegateBinding.h
+++ b/Engine/Plugins/EnhancedInput/Source/EnhancedInput/Public/EnhancedInputActionDelegateBinding.h
@@ -17,7 +17,7 @@ struct ENHANCEDINPUT_API FBlueprintEnhancedInputActionBinding
 	GENERATED_USTRUCT_BODY()
 
 	UPROPERTY()
-	const UInputAction* InputAction = nullptr;
+	TObjectPtr<const UInputAction> InputAction = nullptr;
 
 	UPROPERTY()
 	ETriggerEvent TriggerEvent = ETriggerEvent::None;
diff --git a/Engine/Plugins/EnhancedInput/Source/EnhancedInput/Public/EnhancedPlayerInput.h b/Engine/Plugins/EnhancedInput/Source/EnhancedInput/Public/EnhancedPlayerInput.h
--- a/Engine/Plugins/EnhancedInput/Source/EnhancedInput/Public/EnhancedPlayerInput.h
+++ b/Engine/Plugins/EnhancedInput/Source/EnhancedInput/Public/EnhancedPlayerInput.h
@@ -14,6 +14,7 @@
 // Internal representation containing event variants
 enum class ETriggerEventInternal : uint8;
 enum class EKeyEvent : uint8;
+class UInputMappingContext;
 
 /**
 * UEnhancedPlayerInput : UPlayerInput extensions for enhanced player input system
@@ -80,7 +81,7 @@ private:
 	 * Note: Source reference only. Use EnhancedActionMappings for the actual mappings (with properly instanced triggers/modifiers)
 	 */
 	UPROPERTY(Transient)
-	TMap<const class UInputMappingContext*, int32> AppliedInputContexts;
+	TMap<TObjectPtr<const UInputMappingContext>, int32> AppliedInputContexts;
 
 	/** This player's version of the Action Mappings */
 	UPROPERTY(Transient)
@@ -91,7 +92,7 @@ private:
 
 	/** Tracked action values. Queryable. */
 	UPROPERTY(Transient)
-	mutable TMap<const UInputAction*, FInputActionInstance> ActionInstanceData;
+	mutable TMap<TObjectPtr<const UInputAction>, FInputActionInstance> ActionInstanceData;
 
 	/** Actions which had actuated events at the last call to ProcessInputStack (held/pressed/released) */
 	TSet<const UInputAction*> ActionsWithEventsThisTick;
diff --git a/Engine/Plugins/EnhancedInput/Source/EnhancedInput/Public/InputAction.h b/Engine/Plugins/EnhancedInput/Source/EnhancedInput/Public/InputAction.h
--- a/Engine/Plugins/EnhancedInput/Source/EnhancedInput/Public/InputAction.h
+++ b/Engine/Plugins/EnhancedInput/Source/EnhancedInput/Public/InputAction.h
@@ -54,7 +54,7 @@ public:
 	* All Implicit triggers in this list are met.
 	*/
 	UPROPERTY(EditAnywhere, Instanced, BlueprintReadWrite, Category = Action)
-	TArray<UInputTrigger*> Triggers;
+	TArray<TObjectPtr<UInputTrigger>> Triggers;
 
 	/**
 	* Modifiers are applied to the final action value.
@@ -62,7 +62,7 @@ public:
 	* They are applied on top of any FEnhancedActionKeyMapping modifiers that drove the initial input
 	*/
 	UPROPERTY(EditAnywhere, Instanced, BlueprintReadWrite, Category = Action)
-	TArray<UInputModifier*> Modifiers;
+	TArray<TObjectPtr<UInputModifier>> Modifiers;
 };
 
 // Calculate a collective representation of trigger state from evaluations of all triggers in one or more trigger groups.
@@ -116,7 +116,7 @@ private:
 
 	// The source action that this instance is created from
 	UPROPERTY(BlueprintReadOnly, Category = "Input", Meta = (AllowPrivateAccess))
-	const UInputAction* SourceAction = nullptr;
+	TObjectPtr<const UInputAction> SourceAction = nullptr;
 
 	// Internal trigger states
 	ETriggerState LastTriggerState = ETriggerState::None;
@@ -127,15 +127,15 @@ protected:
 	// TODO: Just hold a duplicate of the UInputAction in here?
 	// TODO: Restrict blueprint access to triggers and modifiers?
 	UPROPERTY(Instanced, BlueprintReadOnly, Category = Config)
-	TArray<UInputTrigger*> Triggers;
+	TArray<TObjectPtr<UInputTrigger>> Triggers;
 
 	UPROPERTY(Instanced, BlueprintReadOnly, Category = Config)
-	TArray<UInputModifier*> Modifiers;
+	TArray<TObjectPtr<UInputModifier>> Modifiers;
 	
 	UPROPERTY(meta = (DeprecatedProperty, DeprecationMessage = "Moved to Modifiers."))
-	TArray<UInputModifier*> PerInputModifiers_DEPRECATED;
+	TArray<TObjectPtr<UInputModifier>> PerInputModifiers_DEPRECATED;
 	UPROPERTY(meta = (DeprecatedProperty, DeprecationMessage = "Moved to Modifiers."))
-	TArray<UInputModifier*> FinalValueModifiers_DEPRECATED;
+	TArray<TObjectPtr<UInputModifier>> FinalValueModifiers_DEPRECATED;
 
 	// Combined value of all inputs mapped to this action
 	FInputActionValue Value;
diff --git a/Engine/Plugins/EnhancedInput/Source/EnhancedInput/Public/InputMappingQuery.h b/Engine/Plugins/EnhancedInput/Source/EnhancedInput/Public/InputMappingQuery.h
--- a/Engine/Plugins/EnhancedInput/Source/EnhancedInput/Public/InputMappingQuery.h
+++ b/Engine/Plugins/EnhancedInput/Source/EnhancedInput/Public/InputMappingQuery.h
@@ -61,6 +61,9 @@ constexpr EMappingQueryIssue TypeMismatch = EMappingQueryIssue::ForcesTypePromot
 constexpr EMappingQueryIssue StandardFatal = EMappingQueryIssue::ReservedByAction | NoCollisions | TypeMismatch;	// Default fatal value for QueryMapKeyIn... function calls.
 }
 
+class UInputMappingContext;
+class UInputAction;
+
 // Potential issue raised with a mapping request
 USTRUCT(BlueprintType)
 struct FMappingQueryIssue
@@ -75,11 +78,11 @@ struct FMappingQueryIssue
 
 	// Input context that contains a blocking action bound to the queried key
 	UPROPERTY(BlueprintReadOnly, Category = Result)
-	const class UInputMappingContext* BlockingContext = nullptr;
+	TObjectPtr<const UInputMappingContext> BlockingContext = nullptr;
 
 	// Action within the input context that caused the blockage
 	UPROPERTY(BlueprintReadOnly, Category = Result)
-	const class UInputAction* BlockingAction = nullptr;
+	TObjectPtr<const UInputAction> BlockingAction = nullptr;
 };
 
 // ************************************************************************************************
diff --git a/Engine/Plugins/EnhancedInput/Source/EnhancedInput/Public/PlayerMappableInputConfig.h b/Engine/Plugins/EnhancedInput/Source/EnhancedInput/Public/PlayerMappableInputConfig.h
--- a/Engine/Plugins/EnhancedInput/Source/EnhancedInput/Public/PlayerMappableInputConfig.h
+++ b/Engine/Plugins/EnhancedInput/Source/EnhancedInput/Public/PlayerMappableInputConfig.h
@@ -66,17 +66,17 @@ protected:
 
 	/** Metadata that can used to store any other related items to your key mapping such as icons, ability assets, etc. */
 	UPROPERTY(EditAnywhere, BlueprintReadOnly, Category = "Input|PlayerMappable")
-	UObject* Metadata = nullptr;
+	TObjectPtr<UObject> Metadata = nullptr;
 
 	/** Mapping contexts that make up this Input Config with their associated priority. */
 	UPROPERTY(EditAnywhere, BlueprintReadOnly, Category = "Input|PlayerMappable")
-	TMap<UInputMappingContext*, int32> Contexts;
+	TMap<TObjectPtr<UInputMappingContext>, int32> Contexts;
 
 public:
 
 	/** Return all the Input Mapping contexts that  */
 	UFUNCTION(BlueprintCallable, Category = "Input|PlayerMappable")
-	const TMap<UInputMappingContext*, int32>& GetMappingContexts() const { return Contexts; }
+	const TMap<TObjectPtr<UInputMappingContext>, int32>& GetMappingContexts() const { return Contexts; }
 
 	UFUNCTION(BlueprintCallable, Category = "Input|PlayerMappable")
 	const FName GetConfigName() const { return ConfigName; }
diff --git a/Engine/Plugins/EnhancedInput/Source/InputEditor/Private/Tests/InputTestFramework.cpp b/Engine/Plugins/EnhancedInput/Source/InputEditor/Private/Tests/InputTestFramework.cpp
--- a/Engine/Plugins/EnhancedInput/Source/InputEditor/Private/Tests/InputTestFramework.cpp
+++ b/Engine/Plugins/EnhancedInput/Source/InputEditor/Private/Tests/InputTestFramework.cpp
@@ -78,11 +78,11 @@ FEnhancedActionKeyMapping& AnActionIsMappedToAKey(UControllablePlayer& PlayerDat
 		{
 			FBindingTargets& BindingTargets = PlayerData.BindingTargets.Emplace(ActionName, FBindingTargets(PlayerData.Player));
 
-			PlayerData.InputComponent->BindAction(Action, ETriggerEvent::Started, BindingTargets.Started, &UInputBindingTarget::MappingListener);
-			PlayerData.InputComponent->BindAction(Action, ETriggerEvent::Ongoing, BindingTargets.Ongoing, &UInputBindingTarget::MappingListener);
-			PlayerData.InputComponent->BindAction(Action, ETriggerEvent::Canceled, BindingTargets.Canceled, &UInputBindingTarget::MappingListener);
-			PlayerData.InputComponent->BindAction(Action, ETriggerEvent::Completed, BindingTargets.Completed, &UInputBindingTarget::MappingListener);
-			PlayerData.InputComponent->BindAction(Action, ETriggerEvent::Triggered, BindingTargets.Triggered, &UInputBindingTarget::MappingListener);
+			PlayerData.InputComponent->BindAction(Action, ETriggerEvent::Started, BindingTargets.Started.Get(), &UInputBindingTarget::MappingListener);
+			PlayerData.InputComponent->BindAction(Action, ETriggerEvent::Ongoing, BindingTargets.Ongoing.Get(), &UInputBindingTarget::MappingListener);
+			PlayerData.InputComponent->BindAction(Action, ETriggerEvent::Canceled, BindingTargets.Canceled.Get(), &UInputBindingTarget::MappingListener);
+			PlayerData.InputComponent->BindAction(Action, ETriggerEvent::Completed, BindingTargets.Completed.Get(), &UInputBindingTarget::MappingListener);
+			PlayerData.InputComponent->BindAction(Action, ETriggerEvent::Triggered, BindingTargets.Triggered.Get(), &UInputBindingTarget::MappingListener);
 		}
 
 		// Initialise mapping in context
@@ -237,13 +237,13 @@ void InputIsTicked(UControllablePlayer& PlayerData, float Delta)
 
 UInputMappingContext* FInputTestHelper::FindContext(UControllablePlayer& Data, FName ContextName)
 {
-	UInputMappingContext** Context = Data.InputContext.Find(ContextName);
+	TObjectPtr<UInputMappingContext>* Context = Data.InputContext.Find(ContextName);
 	return Context ? *Context : nullptr;
 }
 
 UInputAction* FInputTestHelper::FindAction(UControllablePlayer& Data, FName ActionName)
 {
-	UInputAction** Action = Data.InputAction.Find(ActionName);
+	TObjectPtr<UInputAction>* Action = Data.InputAction.Find(ActionName);
 	return Action ? *Action : nullptr;
 }
 
diff --git a/Engine/Plugins/EnhancedInput/Source/InputEditor/Private/Tests/InputTestFramework.h b/Engine/Plugins/EnhancedInput/Source/InputEditor/Private/Tests/InputTestFramework.h
--- a/Engine/Plugins/EnhancedInput/Source/InputEditor/Private/Tests/InputTestFramework.h
+++ b/Engine/Plugins/EnhancedInput/Source/InputEditor/Private/Tests/InputTestFramework.h
@@ -49,19 +49,19 @@ struct FBindingTargets
 	}
 
 	UPROPERTY()
-	UInputBindingTarget* Started = nullptr;
+	TObjectPtr<UInputBindingTarget> Started = nullptr;
 
 	UPROPERTY()
-	UInputBindingTarget* Ongoing = nullptr;
+	TObjectPtr<UInputBindingTarget> Ongoing = nullptr;
 
 	UPROPERTY()
-	UInputBindingTarget* Canceled = nullptr;
+	TObjectPtr<UInputBindingTarget> Canceled = nullptr;
 
 	UPROPERTY()
-	UInputBindingTarget* Completed = nullptr;
+	TObjectPtr<UInputBindingTarget> Completed = nullptr;
 
 	UPROPERTY()
-	UInputBindingTarget* Triggered = nullptr;
+	TObjectPtr<UInputBindingTarget> Triggered = nullptr;
 };
 
 // Mock input subsystems to avoid having to create an actual subsystem + local player + game instance.
@@ -81,7 +81,7 @@ class UControllablePlayer : public UObject
 public:
 
 	UPROPERTY()
-	class APlayerController* Player;
+	TObjectPtr<class APlayerController> Player;
 
 	// These exist on (and have their lifetime managed by) Player, but are downcast here for rapid access.
 	TWeakObjectPtr<UEnhancedPlayerInput> PlayerInput;
@@ -92,11 +92,11 @@ public:
 
 	// Storage for input mapping contexts applied to the player
 	UPROPERTY()
-	TMap<FName, UInputMappingContext*> InputContext;
+	TMap<FName, TObjectPtr<UInputMappingContext>> InputContext;
 
 	// Storage for input actions applied to the player
 	UPROPERTY()
-	TMap<FName, UInputAction*> InputAction;
+	TMap<FName, TObjectPtr<UInputAction>> InputAction;
 
 	// Ensure we don't try to double bind listeners when applying multiple key mappings
 	TSet<const UInputAction*> MappedActionListeners;
