Make the TickForcedInput function protected on the EnhancedInputSubsystemInterface. Without this, the only place that you can tick the input on the Enhanced Input subsystem would be from the Enhanced Input module itself. Ticking from the module is fine if we are making engine modifications and adding more subsystems, but if you want to create your own Tickable subsystem that implements the IEnhancedInputSubsystemInterface then you would not be able to do that without making engine modifications.

This change fixes that!

diff --git a/Engine/Plugins/EnhancedInput/Source/EnhancedInput/Public/EnhancedInputSubsystemInterface.h b/Engine/Plugins/EnhancedInput/Source/EnhancedInput/Public/EnhancedInputSubsystemInterface.h
--- a/Engine/Plugins/EnhancedInput/Source/EnhancedInput/Public/EnhancedInputSubsystemInterface.h
+++ b/Engine/Plugins/EnhancedInput/Source/EnhancedInput/Public/EnhancedInputSubsystemInterface.h
@@ -215,6 +215,14 @@ public:
 	/** Removes all the input mapping contexts inside of this mappable config. */
 	UFUNCTION(BlueprintCallable, BlueprintCosmetic, Category = "Input|PlayerMappable", meta=(AutoCreateRefTerm = "Options"))
 	virtual void RemovePlayerMappableConfig(const UPlayerMappableInputConfig* Config, const FModifyContextOptions& Options = FModifyContextOptions());
+
+protected:
+
+	/**
+	 * Ticks any forced/injected input to the subsystem's player input. If you want support for ApplyForcedInput
+	 * in a subsytem implementing this interface, then you must call this function.
+	 */
+	void TickForcedInput(float DeltaTime);
 	
 private:
 
@@ -223,7 +231,6 @@ private:
 	void ApplyForcedInput(FKey Key, FInputActionValue Value);
 	void RemoveForcedInput(const UInputAction* Action);
 	void RemoveForcedInput(FKey Key);
-	void TickForcedInput(float DeltaTime);
 
 	void InjectChordBlockers(const TArray<int32>& ChordedMappings);
 	bool HasTriggerWith(TFunctionRef<bool(const class UInputTrigger*)> TestFn, const TArray<class UInputTrigger*>& Triggers);
