Add a FEditorScriptExecutionGuard to the Execute function of enhanced input delegates. This makes it possible for enhanced input delegates to be called outside of PIE (during editor time). This can be disabled with "PublicDefinitions.Add("ENABLE_EDITOR_CALLABLE_INPUT_DELEGATES=0")"  if users do not want this enabled in the editor.

diff --git a/Engine/Plugins/EnhancedInput/Source/EnhancedInput/Public/EnhancedInputComponent.h b/Engine/Plugins/EnhancedInput/Source/EnhancedInput/Public/EnhancedInputComponent.h
--- a/Engine/Plugins/EnhancedInput/Source/EnhancedInput/Public/EnhancedInputComponent.h
+++ b/Engine/Plugins/EnhancedInput/Source/EnhancedInput/Public/EnhancedInputComponent.h
@@ -80,12 +80,20 @@ public:
 		return *Delegate;
 	}
 
-	template<typename... TArgs>
+	template<bool bUseScriptGuard = false, typename... TArgs>
 	void Execute(TArgs... Args) const
 	{
 		if (IsBound())
 		{
-			Delegate->Execute(Args...);
+			if (bUseScriptGuard)
+			{
+				FEditorScriptExecutionGuard ScriptGuard;
+				Delegate->Execute(Args...);
+			}
+			else
+			{
+				Delegate->Execute(Args...);	
+			}
 		}
 	}
 };
@@ -243,28 +251,37 @@ public:
 
 // Action event delegate execution by signature.
 
+// By default this behavior is on in the editor only. If you would like to turn this
+// off, then add the following to your Build.cs file:
+//		PublicDefinitions.Add("ENABLE_EDITOR_CALLABLE_INPUT_DELEGATES=0");
+#ifndef ENABLE_EDITOR_CALLABLE_INPUT_DELEGATES
+	#define ENABLE_EDITOR_CALLABLE_INPUT_DELEGATES	WITH_EDITOR
+#elif !WITH_EDITOR
+	#define ENABLE_EDITOR_CALLABLE_INPUT_DELEGATES	0
+#endif
+
 template<>
 inline void FEnhancedInputActionEventDelegateBinding<FEnhancedInputActionHandlerSignature>::Execute(const FInputActionInstance& ActionData) const
 {
-	Delegate.Execute();
+	Delegate.Execute<ENABLE_EDITOR_CALLABLE_INPUT_DELEGATES>();
 }
 
 template<>
 inline void FEnhancedInputActionEventDelegateBinding<FEnhancedInputActionHandlerValueSignature>::Execute(const FInputActionInstance& ActionData) const
 {
-	Delegate.Execute(ActionData.GetValue());
+	Delegate.Execute<ENABLE_EDITOR_CALLABLE_INPUT_DELEGATES>(ActionData.GetValue());
 }
 
 template<>
 inline void FEnhancedInputActionEventDelegateBinding<FEnhancedInputActionHandlerInstanceSignature>::Execute(const FInputActionInstance& ActionData) const
 {
-	Delegate.Execute(ActionData);
+	Delegate.Execute<ENABLE_EDITOR_CALLABLE_INPUT_DELEGATES>(ActionData);
 }
 
 template<>
 inline void FEnhancedInputActionEventDelegateBinding<FEnhancedInputActionHandlerDynamicSignature>::Execute(const FInputActionInstance& ActionData) const
 {
-	Delegate.Execute(ActionData.GetValue(), ActionData.GetElapsedTime(), ActionData.GetTriggeredTime(), ActionData.GetSourceAction());
+	Delegate.Execute<ENABLE_EDITOR_CALLABLE_INPUT_DELEGATES>(ActionData.GetValue(), ActionData.GetElapsedTime(), ActionData.GetTriggeredTime(), ActionData.GetSourceAction());
 }
 
 
