Enhanced Input Editor Subsystem. The Enhanced Input Editor subsystem will allow you to bind to Input Delegates in Editor Utility Actors or other Editor-time objects. The input is driven by an InputProcessor, and this subsystem will only incur any cost to the user if they manually enable it. This is useful for Virtual Production as it will allow physical input devices like cameras to make changes to an object's properties in the level.

diff --git a/Engine/Plugins/EnhancedInput/Source/InputEditor/InputEditor.Build.cs b/Engine/Plugins/EnhancedInput/Source/InputEditor/InputEditor.Build.cs
--- a/Engine/Plugins/EnhancedInput/Source/InputEditor/InputEditor.Build.cs
+++ b/Engine/Plugins/EnhancedInput/Source/InputEditor/InputEditor.Build.cs
@@ -35,7 +35,9 @@ namespace UnrealBuildTool.Rules
                     "Slate",
                     "SlateCore",
                     "UnrealEd",
-					"AssetTools"
+					"AssetTools",
+					"DeveloperSettings",
+					"EditorSubsystem",
                 }
             );
         }
diff --git a/Engine/Plugins/EnhancedInput/Source/InputEditor/Private/EnhancedInputEditorProcessor.cpp b/Engine/Plugins/EnhancedInput/Source/InputEditor/Private/EnhancedInputEditorProcessor.cpp
new file mode 100644
--- /dev/null
+++ b/Engine/Plugins/EnhancedInput/Source/InputEditor/Private/EnhancedInputEditorProcessor.cpp
@@ -0,0 +1,183 @@
+// Copyright Epic Games, Inc. All Rights Reserved.
+
+#include "EnhancedInputEditorProcessor.h"
+#include "EnhancedInputEditorSubsystem.h"
+#include "Framework/Application/SlateApplication.h"
+#include "Editor.h"		// For GEditor
+#include "GameFramework/PlayerInput.h"
+
+void FEnhancedInputEditorProcessor::Tick(const float DeltaTime, FSlateApplication& SlateApp, TSharedRef<ICursor> Cursor)
+{
+	ProcessAccumulatedPointerInput(DeltaTime);
+}
+
+bool FEnhancedInputEditorProcessor::HandleKeyDownEvent(FSlateApplication& SlateApp, const FKeyEvent& InKeyEvent)
+{
+	FInputKeyParams Params = {};
+	Params.Key = InKeyEvent.GetKey();
+	Params.Event = IE_Pressed;
+	Params.Delta.X = 1.0;
+	Params.DeltaTime = SlateApp.GetDeltaTime();
+	Params.NumSamples = Params.Key.IsAnalog() ? 1 : 0;
+	
+	InputKeyToSubsystem(Params);
+	
+	return IInputProcessor::HandleKeyDownEvent(SlateApp, InKeyEvent);
+}
+
+bool FEnhancedInputEditorProcessor::HandleKeyUpEvent(FSlateApplication& SlateApp, const FKeyEvent& InKeyEvent)
+{
+	FInputKeyParams Params = {};
+	Params.Key = InKeyEvent.GetKey();
+	Params.Event = IE_Released;
+	Params.Delta.X = 0.0;
+	Params.DeltaTime = SlateApp.GetDeltaTime();
+	Params.NumSamples = Params.Key.IsAnalog() ? 1 : 0;
+	
+	InputKeyToSubsystem(Params);
+	return IInputProcessor::HandleKeyUpEvent(SlateApp, InKeyEvent);
+}
+
+bool FEnhancedInputEditorProcessor::HandleAnalogInputEvent(FSlateApplication& SlateApp, const FAnalogInputEvent& InAnalogInputEvent)
+{
+	FInputKeyParams Params = {};
+	Params.Key = InAnalogInputEvent.GetKey();
+	Params.Event = IE_Pressed;
+	Params.Delta.X = InAnalogInputEvent.GetAnalogValue();
+	Params.DeltaTime = SlateApp.GetDeltaTime();
+	Params.NumSamples = 1;
+	
+	InputKeyToSubsystem(Params);
+	return IInputProcessor::HandleAnalogInputEvent(SlateApp, InAnalogInputEvent);
+}
+
+bool FEnhancedInputEditorProcessor::HandleMouseMoveEvent(FSlateApplication& SlateApp, const FPointerEvent& MouseEvent)
+{
+	UpdateCachedPointerPosition(MouseEvent);
+	
+	return IInputProcessor::HandleMouseMoveEvent(SlateApp, MouseEvent);
+}
+
+bool FEnhancedInputEditorProcessor::HandleMouseButtonDownEvent(FSlateApplication& SlateApp, const FPointerEvent& MouseEvent)
+{
+	FInputKeyParams Params = {};
+	Params.Key = MouseEvent.GetEffectingButton();
+	Params.Event = IE_Pressed;
+	Params.Delta.X = 1.0;
+	Params.DeltaTime = SlateApp.GetDeltaTime();
+	Params.NumSamples = 0;
+	
+	InputKeyToSubsystem(Params);
+	return IInputProcessor::HandleMouseButtonDownEvent(SlateApp, MouseEvent);
+}
+
+bool FEnhancedInputEditorProcessor::HandleMouseButtonUpEvent(FSlateApplication& SlateApp, const FPointerEvent& MouseEvent)
+{
+	FInputKeyParams Params = {};
+	Params.Key = MouseEvent.GetEffectingButton();
+	Params.Event = IE_Released;
+	Params.Delta.X = 0.0;
+	Params.DeltaTime = SlateApp.GetDeltaTime();
+	Params.NumSamples = 0;
+	
+	InputKeyToSubsystem(Params);
+	return IInputProcessor::HandleMouseButtonUpEvent(SlateApp, MouseEvent);
+}
+
+bool FEnhancedInputEditorProcessor::HandleMouseButtonDoubleClickEvent(FSlateApplication& SlateApp, const FPointerEvent& MouseEvent)
+{
+	FInputKeyParams Params = {};
+	Params.Key = MouseEvent.GetEffectingButton();
+	Params.Event = IE_DoubleClick;
+	Params.Delta.X = 1.0;
+	Params.DeltaTime = SlateApp.GetDeltaTime();
+	Params.NumSamples = 0;
+	
+	InputKeyToSubsystem(Params);
+	return IInputProcessor::HandleMouseButtonDoubleClickEvent(SlateApp, MouseEvent);
+}
+
+bool FEnhancedInputEditorProcessor::HandleMouseWheelOrGestureEvent(FSlateApplication& SlateApp, const FPointerEvent& InWheelEvent, const FPointerEvent* InGestureEvent)
+{
+	const FKey MouseWheelKey = InWheelEvent.GetWheelDelta() < 0 ? EKeys::MouseScrollDown : EKeys::MouseScrollUp;
+
+	// Input the Mouse wheel key events (mouse scroll down or scroll up) as being pressed and released this frame
+	// The SceneViewport Inputs the Mouse Scroll Wheel buttons up and down in the same frame, this replicates that behavior
+	{
+		FInputKeyParams PressedParams = {};
+		PressedParams.Key = MouseWheelKey;
+		PressedParams.Event = IE_Pressed;
+		PressedParams.Delta.X = 1.0;
+		PressedParams.DeltaTime = SlateApp.GetDeltaTime();
+		PressedParams.NumSamples = 0;
+		
+		FInputKeyParams ReleasedParams = PressedParams;
+		ReleasedParams.Event = IE_Released;
+
+		InputKeyToSubsystem(PressedParams);
+		InputKeyToSubsystem(ReleasedParams);	
+	}
+	// Input the wheel axis delta to get the MouseWheelAxis button working
+	{
+		FInputKeyParams Params = {};
+		Params.Key = EKeys::MouseWheelAxis;
+		Params.Delta.X = InWheelEvent.GetWheelDelta();
+		Params.DeltaTime = SlateApp.GetDeltaTime();
+		Params.NumSamples = 1;
+		
+		InputKeyToSubsystem(Params);
+	}
+	
+	return IInputProcessor::HandleMouseWheelOrGestureEvent(SlateApp, InWheelEvent, InGestureEvent);
+}
+
+void FEnhancedInputEditorProcessor::UpdateCachedPointerPosition(const FPointerEvent& MouseEvent)
+{
+	CachedCursorDelta = MouseEvent.GetCursorDelta();
+	
+	++NumCursorSamplesThisFrame.X;
+	++NumCursorSamplesThisFrame.Y;
+}
+
+void FEnhancedInputEditorProcessor::ProcessAccumulatedPointerInput(float DeltaTime)
+{
+	// Input the MouseX value
+	{
+		FInputKeyParams Params = {};
+		Params.Key = EKeys::MouseX;
+		Params.Delta.X = CachedCursorDelta.X;
+		Params.DeltaTime = DeltaTime;
+		Params.NumSamples = NumCursorSamplesThisFrame.X;
+		InputKeyToSubsystem(Params);
+	}
+
+	// Input the MouseY value
+	{
+		FInputKeyParams Params = {};
+		Params.Key = EKeys::MouseY;
+		Params.Delta.X = CachedCursorDelta.Y;
+		Params.DeltaTime = DeltaTime;
+		Params.NumSamples = NumCursorSamplesThisFrame.Y;
+		InputKeyToSubsystem(Params);
+	}
+	
+	NumCursorSamplesThisFrame = FIntPoint::ZeroValue;
+	CachedCursorDelta = FVector2D::ZeroVector;
+}
+
+bool FEnhancedInputEditorProcessor::InputKeyToSubsystem(const FInputKeyParams& Params)
+{
+	// TODO: We may want to process the slate user ID that this input key came from in the future
+	// but for now we don't need to worry about it.
+	
+	// GEditor should always be valid here since the preprocessor is spawned by an Editor Subsystem
+	if (ensure(GEditor))
+	{
+		if (UEnhancedInputEditorSubsystem* Subsystem = GEditor->GetEditorSubsystem<UEnhancedInputEditorSubsystem>())
+		{
+			return Subsystem->InputKey(Params);
+		}	
+	}
+	
+	return false;
+}
diff --git a/Engine/Plugins/EnhancedInput/Source/InputEditor/Private/EnhancedInputEditorSettings.cpp b/Engine/Plugins/EnhancedInput/Source/InputEditor/Private/EnhancedInputEditorSettings.cpp
new file mode 100644
--- /dev/null
+++ b/Engine/Plugins/EnhancedInput/Source/InputEditor/Private/EnhancedInputEditorSettings.cpp
@@ -0,0 +1,43 @@
+// Copyright Epic Games, Inc. All Rights Reserved.
+
+#include "EnhancedInputEditorSettings.h"
+#include "EnhancedPlayerInput.h"
+#include "InputMappingContext.h"
+#include "EnhancedInputEditorSubsystem.h"
+#include "Editor.h"
+
+namespace UE::EnhancedInput::Private
+{
+	namespace ConsoleVariables
+	{
+		static bool bShouldLogAllInputs = false;
+		static FAutoConsoleVariableRef CVarShouldLogAllInputs(
+			TEXT("EnhancedEditorInput.bShouldLogAllInputs"),
+			bShouldLogAllInputs,
+			TEXT("Should each InputKey call be logged?"),
+			ECVF_Default);
+
+		static bool bAutomaticallyStartConsumingInput = false;
+		static FAutoConsoleVariableRef CVarAutomaticallyStartConsumingInput(
+			TEXT("EnhancedEditorInput.bAutomaticallyStartConsumingInput"),
+			bAutomaticallyStartConsumingInput,
+			TEXT("Should the UEnhancedInputEditorSubsystem be started as soon as it is inialized?"),
+			ECVF_Default);	
+	}
+}
+
+////////////////////////////////////////////////////////////////////////
+// UEnhancedInputEditorProjectSettings
+UEnhancedInputEditorProjectSettings::UEnhancedInputEditorProjectSettings(const FObjectInitializer& Initializer)
+	: Super(Initializer)
+	, DefaultEditorInputClass(UEnhancedPlayerInput::StaticClass())
+{
+}
+
+////////////////////////////////////////////////////////////////////////
+// UEnhancedInputEditorSettings
+UEnhancedInputEditorSettings::UEnhancedInputEditorSettings()
+	: bLogAllInput(false)
+	, bAutomaticallyStartConsumingInput(false)
+{
+}
diff --git a/Engine/Plugins/EnhancedInput/Source/InputEditor/Private/EnhancedInputEditorSubsystem.cpp b/Engine/Plugins/EnhancedInput/Source/InputEditor/Private/EnhancedInputEditorSubsystem.cpp
new file mode 100644
--- /dev/null
+++ b/Engine/Plugins/EnhancedInput/Source/InputEditor/Private/EnhancedInputEditorSubsystem.cpp
@@ -0,0 +1,251 @@
+// Copyright Epic Games, Inc. All Rights Reserved.
+
+#include "EnhancedInputEditorSubsystem.h"
+#include "EnhancedPlayerInput.h"
+#include "InputMappingContext.h"
+#include "Components/InputComponent.h"
+#include "Subsystems/UnrealEditorSubsystem.h"
+#include "UnrealEdGlobals.h"
+#include "Editor/UnrealEdEngine.h"
+#include "Editor.h"
+#include "EnhancedInputEditorProcessor.h"
+#include "EnhancedInputEditorSettings.h"
+#include "Framework/Application/SlateApplication.h"
+#include "Kismet/GameplayStatics.h"
+#include "EnhancedInputLibrary.h"
+
+#define LOCTEXT_NAMESPACE "EnhancedInputEditorSubsystem"
+
+DEFINE_LOG_CATEGORY(LogEditorInput);
+
+void UEnhancedInputEditorSubsystem::Initialize(FSubsystemCollectionBase& Collection)
+{
+	Super::Initialize(Collection);
+
+	const UClass* InputClassToUse = GetDefault<UEnhancedInputEditorProjectSettings>()->DefaultEditorInputClass.Get();
+	
+	if (!ensureMsgf(InputClassToUse, TEXT("The UEnhancedInputEditorProjectSettings input class is null! Make sure this has a value set. Using UEnhancedPlayerInput as a fallback.")))
+	{
+		InputClassToUse = UEnhancedPlayerInput::StaticClass();
+	}
+	PlayerInput = NewObject<UEnhancedPlayerInput>(this, InputClassToUse);
+
+	ensureMsgf(PlayerInput, TEXT("UEnhancedInputEditorSubsystem::Initialize failed to create PlayerInput! This subsystem will not tick!"));
+	
+	// Create and register the input preprocessor, this is what will call our "InputKey"
+	// function to drive input instead of a player controller
+	InputPreprocessor = MakeShared<FEnhancedInputEditorProcessor>();
+	FSlateApplication::Get().RegisterInputPreProcessor(InputPreprocessor, 0);
+
+	if (GetDefault<UEnhancedInputEditorSettings>()->bAutomaticallyStartConsumingInput)
+	{
+		StartConsumingInput();
+	}
+}
+
+void UEnhancedInputEditorSubsystem::Deinitialize()
+{
+	Super::Deinitialize();
+
+	if (FSlateApplication::IsInitialized())
+	{
+		FSlateApplication::Get().UnregisterInputPreProcessor(InputPreprocessor);
+	}
+	
+	InputPreprocessor.Reset();
+	CurrentInputStack.Empty();
+	PlayerInput = nullptr;
+}
+
+UWorld* UEnhancedInputEditorSubsystem::GetTickableGameObjectWorld() const
+{
+	return GetWorld();
+}
+
+ETickableTickType UEnhancedInputEditorSubsystem::GetTickableTickType() const
+{
+	// Don't let the CDO ever tick
+	return IsTemplate() ? ETickableTickType::Never : ETickableTickType::Conditional;
+}
+
+bool UEnhancedInputEditorSubsystem::IsAllowedToTick() const
+{
+	// Only tick if we have created a valid player input (i.e. we have been initalized)
+	return PlayerInput != nullptr && bIsCurrentlyConsumingInput;
+}
+
+void UEnhancedInputEditorSubsystem::Tick(float DeltaTime)
+{
+	if (!ensure(PlayerInput))
+	{
+		UE_LOG(LogEditorInput, Error, TEXT("UEnhancedInputEditorSubsystem is ticking before it has created a PlayerInput!"));
+		return;
+	}
+
+	FModifyContextOptions Options;
+	Options.bForceImmediately = true;
+	
+	// Rebuild the control mappings and tick any forced input that may have been injected
+	RequestRebuildControlMappings(Options);
+	TickForcedInput(DeltaTime);
+
+	static TArray<UInputComponent*> InputStack;
+
+	// Build the input stack
+	{
+		for (int32 i = 0; i < CurrentInputStack.Num(); ++i)
+		{
+			if (UInputComponent* IC = CurrentInputStack[i].Get())
+			{
+				InputStack.Push(IC);
+			}
+			else
+			{
+				CurrentInputStack.RemoveAt(i--);
+			}
+		}
+	}
+	
+	// Process input stack on the player input
+	PlayerInput->Tick(DeltaTime);
+	PlayerInput->ProcessInputStack(InputStack, DeltaTime, false);
+
+	InputStack.Reset();
+}
+
+void UEnhancedInputEditorSubsystem::PushInputComponent(UInputComponent* InInputComponent)
+{
+	if (InInputComponent)
+	{
+		if (GetDefault<UEnhancedInputEditorSettings>()->bLogAllInput)
+		{
+			UE_LOG(LogEditorInput, Log, TEXT("PUSHING AN INPUT COMPONENT! '%s'"), *InInputComponent->GetName());
+		}
+
+		bool bPushed = false;
+		CurrentInputStack.RemoveSingle(InInputComponent);
+		for (int32 Index = CurrentInputStack.Num() - 1; Index >= 0; --Index)
+		{
+			UInputComponent* IC = CurrentInputStack[Index].Get();
+			if (IC == nullptr)
+			{
+				CurrentInputStack.RemoveAt(Index);
+			}
+			else if (IC->Priority <= InInputComponent->Priority)
+			{
+				CurrentInputStack.Insert(InInputComponent, Index + 1);
+				bPushed = true;
+				break;
+			}
+		}
+		if (!bPushed)
+		{
+			CurrentInputStack.Insert(InInputComponent, 0);
+			RequestRebuildControlMappings();
+		}
+	}
+	else
+	{
+		UE_LOG(LogEditorInput, Error, TEXT("Attempted to push a null Input Component to the Enhanced Input Editor Subsystem!"));
+	}
+}
+
+bool UEnhancedInputEditorSubsystem::PopInputComponent(UInputComponent* InInputComponent)
+{
+	if (InInputComponent)
+	{
+		if (CurrentInputStack.RemoveSingle(InInputComponent) > 0)
+		{
+			InInputComponent->ClearBindingValues();
+			RequestRebuildControlMappings();
+			return true;
+		}
+	}
+	return false;
+}
+
+void UEnhancedInputEditorSubsystem::StartConsumingInput()
+{
+	bIsCurrentlyConsumingInput = true;
+	
+	AddDefaultMappingContexts();
+
+	RequestRebuildControlMappings();
+}
+
+void UEnhancedInputEditorSubsystem::StopConsumingInput()
+{
+	bIsCurrentlyConsumingInput = false;
+
+	RemoveDefaultMappingContexts();
+
+	RequestRebuildControlMappings();
+}
+
+UEnhancedPlayerInput* UEnhancedInputEditorSubsystem::GetPlayerInput() const
+{
+	return PlayerInput.Get();
+}
+
+UWorld* UEnhancedInputEditorSubsystem::GetWorld() const
+{
+	if (UUnrealEditorSubsystem* UnrealEditorSubsystem = GEditor->GetEditorSubsystem<UUnrealEditorSubsystem>())
+	{
+		if (GUnrealEd->IsPlayingSessionInEditor())
+		{
+			return UnrealEditorSubsystem->GetGameWorld();
+		}
+		else
+		{
+			return UnrealEditorSubsystem->GetEditorWorld();
+		}
+	}
+
+	return Super::GetWorld();
+}
+
+bool UEnhancedInputEditorSubsystem::InputKey(const FInputKeyParams& Params)
+{
+	if (PlayerInput)
+	{
+		if (bIsCurrentlyConsumingInput)
+		{
+			if (GetDefault<UEnhancedInputEditorSettings>()->bLogAllInput)
+			{
+				UE_LOG(LogEditorInput, Log, TEXT("EI Editor Subsystem InputKey : [%s]"), *Params.Key.ToString());
+			}
+			return PlayerInput->InputKey(Params);
+		}
+		return false;
+	}
+
+	ensureAlwaysMsgf(false, TEXT("Attempting to input a key to the EnhancedInputEditorSubsystem, but there is no Player Input!"));
+	return false;
+}
+
+void UEnhancedInputEditorSubsystem::AddDefaultMappingContexts()
+{
+	for (const FDefaultContextSetting& ContextSetting : GetMutableDefault<UEnhancedInputEditorProjectSettings>()->DefaultMappingContexts)
+	{
+		if (const UInputMappingContext* IMC = ContextSetting.InputMappingContext.LoadSynchronous())
+		{
+			if (!HasMappingContext(IMC))
+			{
+				AddMappingContext(IMC, ContextSetting.Priority);	
+			}
+		}
+	}
+}
+
+void UEnhancedInputEditorSubsystem::RemoveDefaultMappingContexts()
+{
+	for (const FDefaultContextSetting& ContextSetting : GetMutableDefault<UEnhancedInputEditorProjectSettings>()->DefaultMappingContexts)
+	{
+		if (const UInputMappingContext* IMC = ContextSetting.InputMappingContext.LoadSynchronous())
+		{
+			RemoveMappingContext(IMC);
+		}
+	}
+}
+
+#undef LOCTEXT_NAMESPACE
diff --git a/Engine/Plugins/EnhancedInput/Source/InputEditor/Public/EnhancedInputEditorProcessor.h b/Engine/Plugins/EnhancedInput/Source/InputEditor/Public/EnhancedInputEditorProcessor.h
new file mode 100644
--- /dev/null
+++ b/Engine/Plugins/EnhancedInput/Source/InputEditor/Public/EnhancedInputEditorProcessor.h
@@ -0,0 +1,51 @@
+// Copyright Epic Games, Inc. All Rights Reserved.
+
+#pragma once
+
+#include "CoreMinimal.h"
+#include "Framework/Application/IInputProcessor.h"
+
+struct FInputKeyParams;
+
+/**
+ * The Enhanced Input Editor Processor is used to pass InputKey events to the
+ * Enhanced Input Editor Subsystem.
+ *
+ * It will not steal input, and all the functions here will return "False"
+ * so that other Input Processors still run with all their normal considerations. 
+ */
+class FEnhancedInputEditorProcessor : public IInputProcessor
+{
+public:
+	//~ Begin IInputProcessor interface
+	virtual void Tick(const float DeltaTime, FSlateApplication& SlateApp, TSharedRef<ICursor> Cursor) override;
+	virtual bool HandleKeyDownEvent(FSlateApplication& SlateApp, const FKeyEvent& InKeyEvent) override;
+	virtual bool HandleKeyUpEvent(FSlateApplication& SlateApp, const FKeyEvent& InKeyEvent) override;
+	virtual bool HandleAnalogInputEvent(FSlateApplication& SlateApp, const FAnalogInputEvent& InAnalogInputEvent) override;
+	virtual bool HandleMouseMoveEvent(FSlateApplication& SlateApp, const FPointerEvent& MouseEvent) override;
+	virtual bool HandleMouseButtonDownEvent( FSlateApplication& SlateApp, const FPointerEvent& MouseEvent) override;
+	virtual bool HandleMouseButtonUpEvent( FSlateApplication& SlateApp, const FPointerEvent& MouseEvent) override;
+	virtual bool HandleMouseButtonDoubleClickEvent(FSlateApplication& SlateApp, const FPointerEvent& MouseEvent) override;
+	virtual bool HandleMouseWheelOrGestureEvent(FSlateApplication& SlateApp, const FPointerEvent& InWheelEvent, const FPointerEvent* InGestureEvent) override;
+	virtual const TCHAR* GetDebugName() const { return TEXT("EnhancedInputEditorPreprocessor"); }
+	//~ End IInputProcessor interface
+
+protected:
+
+	void UpdateCachedPointerPosition(const FPointerEvent& MouseEvent);
+
+	/** Called on Tick to process any accumulated mouse input */
+	void ProcessAccumulatedPointerInput(float DeltaTime);
+	
+	/** Tells the Enhanced Input Editor subsystem about the given key event */
+	bool InputKeyToSubsystem(const FInputKeyParams& Params);
+
+	/**
+	 * The number of cursor samples that happened this frame. The X and Y will be incremented when there is a mouse move event this frame
+	 * and it will be reset on ProcessAccumulatedPointerInput.
+	 */
+	FIntPoint NumCursorSamplesThisFrame = FIntPoint::ZeroValue;
+	
+	/** The value of the mouse cursor from the most recent mouse event */
+	FVector2D CachedCursorDelta = FVector2D::ZeroVector;
+};
\ No newline at end of file
diff --git a/Engine/Plugins/EnhancedInput/Source/InputEditor/Public/EnhancedInputEditorSettings.h b/Engine/Plugins/EnhancedInput/Source/InputEditor/Public/EnhancedInputEditorSettings.h
new file mode 100644
--- /dev/null
+++ b/Engine/Plugins/EnhancedInput/Source/InputEditor/Public/EnhancedInputEditorSettings.h
@@ -0,0 +1,68 @@
+// Copyright Epic Games, Inc. All Rights Reserved.
+
+#pragma once
+
+#include "CoreMinimal.h"
+#include "Engine/DeveloperSettings.h"
+#include "Engine/DeveloperSettingsBackedByCVars.h"
+
+#include "EnhancedInputEditorSettings.generated.h"
+
+class UEnhancedPlayerInput;
+class UInputMappingContext;
+class UEnhancedInputEditorSubsystem;
+
+USTRUCT()
+struct FDefaultContextSetting
+{
+	GENERATED_BODY()
+
+	/** Input Mapping Context that should be Added to the EnhancedInputEditorSubsystem when it starts listening for input */
+	UPROPERTY(EditAnywhere, Category = "Input")
+	TSoftObjectPtr<const UInputMappingContext> InputMappingContext = nullptr;
+
+	/** The prioirty that should be given to this mapping context when it is added */
+	UPROPERTY(EditAnywhere, Category = "Input")
+	int32 Priority = 0;
+};
+
+/** Settings for the Enhanced Input Editor Subsystem that are persistent between a project's users */
+UCLASS(config = Input, defaultconfig)
+class UEnhancedInputEditorProjectSettings : public UDeveloperSettings
+{
+	GENERATED_BODY()
+
+public: 
+	UEnhancedInputEditorProjectSettings(const FObjectInitializer& Initializer);
+	
+	/** The default player input class that the Enhanced Input Editor subsystem will use. */
+	UPROPERTY(config, EditAnywhere, NoClear, Category = Default)
+	TSoftClassPtr<UEnhancedPlayerInput> DefaultEditorInputClass;
+
+	/** Array of any input mapping contexts that you want to always be applied to the Enhanced Input Editor Subsystem. */
+	UPROPERTY(config, EditAnywhere, Category = Default)
+	TArray<FDefaultContextSetting> DefaultMappingContexts;
+};
+
+/**
+ * A collection of useful indivudal user settings when using the EnhancedInputEditorSubsystem.
+ */
+UCLASS(config=EditorPerProjectUserSettings, meta=(DisplayName="Enhanced Input Editor Settings"))
+class UEnhancedInputEditorSettings : public UDeveloperSettingsBackedByCVars
+{
+	GENERATED_BODY()
+public:
+
+	UEnhancedInputEditorSettings();
+
+	/**
+	 * If true then the Enhanced Input Editor subsystem will log all input that is being processed by it (keypresses, analog values, etc)
+	 * Note: This can produce A LOT of logs, so only use this if you are debugging something.
+	 */
+	UPROPERTY(config, EditAnywhere, Category = Logging, meta=(ConsoleVariable="EnhancedEditorInput.bShouldLogAllInputs"))
+	uint8 bLogAllInput : 1;
+	
+	/** If true, then the UEnhancedInputEditorSubsystem will be started when it is initalized */
+	UPROPERTY(config, EditAnywhere, Category = Editor, meta=(ConsoleVariable="EnhancedEditorInput.bAutomaticallyStartConsumingInput"))
+	uint8 bAutomaticallyStartConsumingInput : 1;
+};
\ No newline at end of file
diff --git a/Engine/Plugins/EnhancedInput/Source/InputEditor/Public/EnhancedInputEditorSubsystem.h b/Engine/Plugins/EnhancedInput/Source/InputEditor/Public/EnhancedInputEditorSubsystem.h
new file mode 100644
--- /dev/null
+++ b/Engine/Plugins/EnhancedInput/Source/InputEditor/Public/EnhancedInputEditorSubsystem.h
@@ -0,0 +1,102 @@
+// Copyright Epic Games, Inc. All Rights Reserved.
+
+#pragma once
+
+#include "CoreMinimal.h"
+#include "EnhancedInputSubsystemInterface.h"
+#include "EditorSubsystem.h"
+#include "Tickable.h"
+#include "GameFramework/PlayerInput.h"
+#include "EnhancedInputEditorSubsystem.generated.h"
+
+DECLARE_LOG_CATEGORY_EXTERN(LogEditorInput, Log, All);
+
+class UInputComponent;
+class UEnhancedPlayerInput;
+class FEnhancedInputEditorProcessor;
+
+/**
+ * The Enhanced Input Editor Subsystem can be used to process input outside of PIE within the editor.
+ * Calling StartConsumingInput will allow the input preprocessor to drive Input Action delegates
+ * to be fired in the editor.
+ *
+ * This allows you to hook up Input Action delegates in Editor Utilities to make editor tools driven by
+ * input.
+ */
+UCLASS()
+class UEnhancedInputEditorSubsystem : public UEditorSubsystem, public IEnhancedInputSubsystemInterface, public FTickableGameObject
+{
+	GENERATED_BODY()
+
+public:
+
+	//~ Begin USubsystem interface
+	virtual void Initialize(FSubsystemCollectionBase& Collection) override;
+	virtual void Deinitialize() override;
+	//~ End USubsystem interface
+	
+	//~ Begin FTickableGameObject interface
+	virtual UWorld* GetTickableGameObjectWorld() const override;
+	virtual bool IsTickableInEditor() const { return true; }
+	virtual ETickableTickType GetTickableTickType() const override;
+	virtual bool IsAllowedToTick() const override;
+	virtual void Tick(float DeltaTime) override;
+	TStatId GetStatId() const override { RETURN_QUICK_DECLARE_CYCLE_STAT(UEnhancedInputEditorSubsystem, STATGROUP_Tickables); }
+	//~ End FTickableGameObject interface
+
+	//~ UObject interface
+	virtual UWorld* GetWorld() const override;
+	//~ End UObject interface
+	
+	//~ Begin IEnhancedInputSubsystemInterface
+	virtual UEnhancedPlayerInput* GetPlayerInput() const override;
+	//~ End IEnhancedInputSubsystemInterface
+
+	/** Pushes this input component onto the stack to be processed by this subsystem's tick function */
+	UFUNCTION(BlueprintCallable, Category = "Input|Editor")
+	void PushInputComponent(UInputComponent* InInputComponent);
+
+	/** Removes this input component onto the stack to be processed by this subsystem's tick function */
+	UFUNCTION(BlueprintCallable, Category = "Input|Editor")
+	bool PopInputComponent(UInputComponent* InInputComponent);
+
+	/** Start the consumption of input messages in this subsystem. This is required to have any Input Action delegates be fired. */
+	UFUNCTION(BlueprintCallable, Category = "Input|Editor")
+	void StartConsumingInput();
+
+	/** Tells this subsystem to stop ticking and consuming any input. This will stop any Input Action Delegates from being called. */
+	UFUNCTION(BlueprintCallable, Category = "Input|Editor")
+	void StopConsumingInput();
+
+	/** Returns true if this subsystem is currently consuming input */
+	UFUNCTION(BlueprintCallable, BlueprintPure, Category = "Input|Editor")
+	bool IsConsumingInput() const { return bIsCurrentlyConsumingInput; }
+
+	/** Inputs a key on this subsystem's player input which can then be processed as normal during Tick. */
+	bool InputKey(const FInputKeyParams& Params);
+	
+	/** Adds all the default mapping contexts from the UEnhancedInputEditorSettings */
+	void AddDefaultMappingContexts();
+
+	/** Removes all the default mapping contexts from the UEnhancedInputEditorSettings */
+	void RemoveDefaultMappingContexts();
+	
+private:
+
+	/** The player input that is processing the input within this subsystem */
+	UPROPERTY()
+	TObjectPtr<UEnhancedPlayerInput> PlayerInput = nullptr;
+
+	/**
+	 * Input processor that is created on Initalize. This will take input from the editor and pass it through
+	 * to this subsystem via InputKey.
+	 */
+	TSharedPtr<FEnhancedInputEditorProcessor> InputPreprocessor = nullptr;
+	
+	/** If true, then this subsystem will Tick and process input delegates. */
+	bool bIsCurrentlyConsumingInput = false;
+	
+	/** Internal. This is the current stack of InputComponents that is being processed by the PlayerInput. */
+	UPROPERTY(Transient)
+	TArray<TWeakObjectPtr<UInputComponent>> CurrentInputStack;
+};
\ No newline at end of file
