Only create the UEnhancedInputEditorSubsystem if we are WITH_EDITOR, and the FSlateApplication has been initalized. Otherwise we can't use it since there would be nothing to register the input preprocessor with. This solves issues with headless cooks crashing on initalize.

#rnx

diff --git a/Engine/Plugins/EnhancedInput/Source/InputEditor/Private/EnhancedInputEditorSubsystem.cpp b/Engine/Plugins/EnhancedInput/Source/InputEditor/Private/EnhancedInputEditorSubsystem.cpp
--- a/Engine/Plugins/EnhancedInput/Source/InputEditor/Private/EnhancedInputEditorSubsystem.cpp
+++ b/Engine/Plugins/EnhancedInput/Source/InputEditor/Private/EnhancedInputEditorSubsystem.cpp
@@ -34,8 +34,11 @@ void UEnhancedInputEditorSubsystem::Initialize(FSubsystemCollectionBase& Collect
 	
 	// Create and register the input preprocessor, this is what will call our "InputKey"
 	// function to drive input instead of a player controller
-	InputPreprocessor = MakeShared<FEnhancedInputEditorProcessor>();
-	FSlateApplication::Get().RegisterInputPreProcessor(InputPreprocessor, 0);
+	if (ensure(FSlateApplication::IsInitialized()))
+	{
+		InputPreprocessor = MakeShared<FEnhancedInputEditorProcessor>();
+		FSlateApplication::Get().RegisterInputPreProcessor(InputPreprocessor, 0);	
+	}
 
 	if (GetDefault<UEnhancedInputEditorSettings>()->bAutomaticallyStartConsumingInput)
 	{
@@ -57,6 +60,23 @@ void UEnhancedInputEditorSubsystem::Deinitialize()
 	PlayerInput = nullptr;
 }
 
+bool UEnhancedInputEditorSubsystem::ShouldCreateSubsystem(UObject* Outer) const
+{
+	// We don't ever want to create this subsystem outside of the editor
+	// If you are outside the editor, then your input delegates would be driven by
+	// the regular Enhanced Input subsystem.
+#if !WITH_EDITOR
+	return false;
+#endif
+	
+	if (!FSlateApplication::IsInitialized())
+	{
+		return false;
+	}
+	
+	return Super::ShouldCreateSubsystem(Outer);
+}
+
 UWorld* UEnhancedInputEditorSubsystem::GetTickableGameObjectWorld() const
 {
 	return GetWorld();
diff --git a/Engine/Plugins/EnhancedInput/Source/InputEditor/Public/EnhancedInputEditorSubsystem.h b/Engine/Plugins/EnhancedInput/Source/InputEditor/Public/EnhancedInputEditorSubsystem.h
--- a/Engine/Plugins/EnhancedInput/Source/InputEditor/Public/EnhancedInputEditorSubsystem.h
+++ b/Engine/Plugins/EnhancedInput/Source/InputEditor/Public/EnhancedInputEditorSubsystem.h
@@ -33,6 +33,7 @@ public:
 	//~ Begin USubsystem interface
 	virtual void Initialize(FSubsystemCollectionBase& Collection) override;
 	virtual void Deinitialize() override;
+	virtual bool ShouldCreateSubsystem(UObject* Outer) const override;
 	//~ End USubsystem interface
 	
 	//~ Begin FTickableGameObject interface
