Fix horde error where we needed "WITH_AUTOMATION_TESTS" in the test framework.

#rnx

diff --git a/Engine/Plugins/EnhancedInput/Source/InputEditor/Private/Tests/InputTestFramework.cpp b/Engine/Plugins/EnhancedInput/Source/InputEditor/Private/Tests/InputTestFramework.cpp
--- a/Engine/Plugins/EnhancedInput/Source/InputEditor/Private/Tests/InputTestFramework.cpp
+++ b/Engine/Plugins/EnhancedInput/Source/InputEditor/Private/Tests/InputTestFramework.cpp
@@ -20,7 +20,10 @@ FMockedEnhancedInputSubsystem::FMockedEnhancedInputSubsystem(const UControllable
 
 UWorld* AnEmptyWorld()
 {
+#if WITH_AUTOMATION_TESTS
 	return FAutomationEditorCommonUtils::CreateNewMap();
+#endif
+	return nullptr;
 }
 
 UControllablePlayer& AControllablePlayer(UWorld* World)
