Change FMath::Min/Max to fix two potholes with the current definitions

1. On custom types, they require overloads for operator <= and >= respectively; it's more idiomatic in C++ code for ordering predicates to rely solely on <
2. The current definition results in highly suboptimal code generation for floating-point types on x86 when compiling in IEEE compliant mode: https://godbolt.org/z/3zsYKaasc

In short, x86 min/max instructions are equivalent to "(a < b) ? a : b" and "(a > b) ? a : b", respectively. In the specific case where a=+0, b=-0 (IEEE signed zeros), then in comparisons +0 == -0, so we have both (a >= b) and (a <= b), but neither (a > b) nor (b < a). As a result, the original expressions for FMath::{Min,Max} return a (+0), while the changed expressions returns b (-0) in this particular case. (The a=-0, b=+0 case is analogous.)

Most code in UE currently compiles with /fp:fast or -ffast-math, which makes the compiler not respect signed zero behavior; see behavior with -fno-signed-zeros in the Compiler Explorer link. Because so much gets compiled with fast math semantics, behavior in this case was not reliable anyway. We might as well pick the definition that is efficient on x86 even without fast math.

diff --git a/Engine/Plugins/EnhancedInput/Source/EnhancedInput/Public/InputAction.h b/Engine/Plugins/EnhancedInput/Source/EnhancedInput/Public/InputAction.h
--- a/Engine/Plugins/EnhancedInput/Source/EnhancedInput/Public/InputAction.h
+++ b/Engine/Plugins/EnhancedInput/Source/EnhancedInput/Public/InputAction.h
@@ -84,6 +84,7 @@ struct FTriggerStateTracker
 	bool GetMappingTriggerApplied() const { return bMappingTriggerApplied; }
 
 	bool operator>=(const FTriggerStateTracker& Other) const { return GetState() >= Other.GetState(); }
+	bool operator< (const FTriggerStateTracker& Other) const { return GetState() <  Other.GetState(); }
 
 	// TODO: Hacky. This is the state we should return if we have evaluated no valid triggers. Set during action evaluation based on final ModifiedValue.
 	void SetStateForNoTriggers(ETriggerState State);
diff --git a/Engine/Source/Runtime/Core/Public/GenericPlatform/GenericPlatformMath.h b/Engine/Source/Runtime/Core/Public/GenericPlatform/GenericPlatformMath.h
--- a/Engine/Source/Runtime/Core/Public/GenericPlatform/GenericPlatformMath.h
+++ b/Engine/Source/Runtime/Core/Public/GenericPlatform/GenericPlatformMath.h
@@ -963,14 +963,14 @@ struct FGenericPlatformMath
 	template< class T > 
 	static CONSTEXPR FORCEINLINE T Max( const T A, const T B )
 	{
-		return (A>=B) ? A : B;
+		return (B < A) ? A : B;
 	}
 
 	/** Returns lower value in a generic way */
 	template< class T > 
 	static CONSTEXPR FORCEINLINE T Min( const T A, const T B )
 	{
-		return (A<=B) ? A : B;
+		return (A < B) ? A : B;
 	}
 
 	// Allow mixing of float types to promote to highest precision type
