Mark Enhanced Input as non-beta because it will be the default on option for 5.1

#rnx

diff --git a/Engine/Plugins/EnhancedInput/EnhancedInput.uplugin b/Engine/Plugins/EnhancedInput/EnhancedInput.uplugin
--- a/Engine/Plugins/EnhancedInput/EnhancedInput.uplugin
+++ b/Engine/Plugins/EnhancedInput/EnhancedInput.uplugin
@@ -12,7 +12,7 @@
 	"SupportURL": "",
 	"EnabledByDefault": true,
 	"CanContainContent": false,
-	"IsBetaVersion": true,
+	"IsBetaVersion": false,
 	"Installed": false,
 	"Modules": [
 		{
