Fix static initialization using /Script/CoreUObject before it is created causing crash in monolithic editor

diff --git a/Engine/Plugins/EnhancedInput/Source/InputBlueprintNodes/Private/K2Node_GetInputActionValue.cpp b/Engine/Plugins/EnhancedInput/Source/InputBlueprintNodes/Private/K2Node_GetInputActionValue.cpp
--- a/Engine/Plugins/EnhancedInput/Source/InputBlueprintNodes/Private/K2Node_GetInputActionValue.cpp
+++ b/Engine/Plugins/EnhancedInput/Source/InputBlueprintNodes/Private/K2Node_GetInputActionValue.cpp
@@ -45,30 +45,34 @@ struct FValueTypeData
 	UScriptStruct* SubCategoryObject;
 };
 
-static TMap<EInputActionValueType, FValueTypeData> ValueLookups =
+const TMap<EInputActionValueType, FValueTypeData>& GetValueLookups()
 {
-	{ EInputActionValueType::Boolean, FValueTypeData(UEdGraphSchema_K2::PC_Boolean) },
-	{ EInputActionValueType::Axis1D, FValueTypeData(UEdGraphSchema_K2::PC_Real, UEdGraphSchema_K2::PC_Float) },
-	{ EInputActionValueType::Axis2D, FValueTypeData(UEdGraphSchema_K2::PC_Struct, NAME_None, TBaseStructure<FVector2D>::Get()) },
-	{ EInputActionValueType::Axis3D, FValueTypeData(UEdGraphSchema_K2::PC_Struct, NAME_None, TBaseStructure<FVector>::Get()) },
-};
+	static const TMap<EInputActionValueType, FValueTypeData> ValueLookups =
+	{
+		{ EInputActionValueType::Boolean, FValueTypeData(UEdGraphSchema_K2::PC_Boolean) },
+		{ EInputActionValueType::Axis1D, FValueTypeData(UEdGraphSchema_K2::PC_Real, UEdGraphSchema_K2::PC_Float) },
+		{ EInputActionValueType::Axis2D, FValueTypeData(UEdGraphSchema_K2::PC_Struct, NAME_None, TBaseStructure<FVector2D>::Get()) },
+		{ EInputActionValueType::Axis3D, FValueTypeData(UEdGraphSchema_K2::PC_Struct, NAME_None, TBaseStructure<FVector>::Get()) },
+	};
+	return ValueLookups;
+}
 
 FName UK2Node_GetInputActionValue::GetValueCategory(const UInputAction* InputAction)
 {
 	EInputActionValueType Type = InputAction ? InputAction->ValueType : EInputActionValueType::Boolean;
-	return ValueLookups[Type].Category;
+	return GetValueLookups()[Type].Category;
 }
 
 FName UK2Node_GetInputActionValue::GetValueSubCategory(const UInputAction* InputAction)
 {
 	EInputActionValueType Type = InputAction ? InputAction->ValueType : EInputActionValueType::Boolean;
-	return ValueLookups[Type].SubCategory;
+	return GetValueLookups()[Type].SubCategory;
 }
 
 UScriptStruct* UK2Node_GetInputActionValue::GetValueSubCategoryObject(const UInputAction* InputAction)
 {
 	EInputActionValueType Type = InputAction ? InputAction->ValueType : EInputActionValueType::Boolean;
-	return ValueLookups[Type].SubCategoryObject;
+	return GetValueLookups()[Type].SubCategoryObject;
 }
 
 void UK2Node_GetInputActionValue::AllocateDefaultPins()
