Add Enhanced Input developer and Platform settings.

Create a "UEnhancedInputPlatformData" base class that can be subclassed in a blueprint and placed on specific platforms. By default, this object has a map of Input Contexts that will "redirect" one context to another. This makes it possible to swap out Input Mapping Contexts on specific platforms, with an asset that will simply be added to the input ini of the game.

diff --git a/Engine/Plugins/EnhancedInput/Source/EnhancedInput/EnhancedInput.Build.cs b/Engine/Plugins/EnhancedInput/Source/EnhancedInput/EnhancedInput.Build.cs
--- a/Engine/Plugins/EnhancedInput/Source/EnhancedInput/EnhancedInput.Build.cs
+++ b/Engine/Plugins/EnhancedInput/Source/EnhancedInput/EnhancedInput.Build.cs
@@ -19,7 +19,8 @@ namespace UnrealBuildTool.Rules
                     "Engine",
 					"InputCore",
 					"Slate",
-                    "SlateCore"
+                    "SlateCore",
+                    "DeveloperSettings",
                 }
             );
         }
diff --git a/Engine/Plugins/EnhancedInput/Source/EnhancedInput/Private/EnhancedInputDeveloperSettings.cpp b/Engine/Plugins/EnhancedInput/Source/EnhancedInput/Private/EnhancedInputDeveloperSettings.cpp
new file mode 100644
--- /dev/null
+++ b/Engine/Plugins/EnhancedInput/Source/EnhancedInput/Private/EnhancedInputDeveloperSettings.cpp
@@ -0,0 +1,10 @@
+// Copyright Epic Games, Inc. All Rights Reserved.
+
+#include "EnhancedInputDeveloperSettings.h"
+#include "EnhancedInputPlatformSettings.h"
+
+UEnhancedInputDeveloperSettings::UEnhancedInputDeveloperSettings(const FObjectInitializer& Initializer)
+	: Super(Initializer)
+{
+	PlatformSettings.Initialize(UEnhancedInputPlatformSettings::StaticClass());
+}
\ No newline at end of file
diff --git a/Engine/Plugins/EnhancedInput/Source/EnhancedInput/Private/EnhancedInputPlatformSettings.cpp b/Engine/Plugins/EnhancedInput/Source/EnhancedInput/Private/EnhancedInputPlatformSettings.cpp
new file mode 100644
--- /dev/null
+++ b/Engine/Plugins/EnhancedInput/Source/EnhancedInput/Private/EnhancedInputPlatformSettings.cpp
@@ -0,0 +1,109 @@
+// Copyright Epic Games, Inc. All Rights Reserved.
+
+#include "EnhancedInputPlatformSettings.h"
+#include "InputMappingContext.h"
+
+#define LOCTEXT_NAMESPACE "EnhancedInputPlatformSettings"
+
+//////////////////////////////////////////////////////////////////////////
+// UEnhancedInputPlatformData
+
+#if WITH_EDITOR
+EDataValidationResult UEnhancedInputPlatformData::IsDataValid(TArray<FText>& ValidationErrors)
+{
+	EDataValidationResult Result = Super::IsDataValid(ValidationErrors);
+
+	for (const TPair<TObjectPtr<const UInputMappingContext>, TObjectPtr<const UInputMappingContext>> Pair : MappingContextRedirects)
+	{
+		if (Pair.Key.IsNull())
+		{
+			Result = CombineDataValidationResults(Result, EDataValidationResult::Invalid);
+			
+			FFormatNamedArguments Args;
+			Args.Add(TEXT("AssetPath"), FText::FromString(GetPathName()));
+			const FText NullKeyError = FText::Format(LOCTEXT("NullKeyError", "'{AssetPath}' does not have a valid key in the MappingContextRedirects!"), Args);
+			ValidationErrors.Emplace(NullKeyError);
+		}
+		
+		if (Pair.Value.IsNull())
+		{
+			Result = CombineDataValidationResults(Result, EDataValidationResult::Invalid);
+			
+			FFormatNamedArguments Args;
+			Args.Add(TEXT("AssetPath"), FText::FromString(GetPathName()));
+			const FText NullValueError = FText::Format(LOCTEXT("NullValueError", "'{AssetPath}' does not have a valid value in the MappingContextRedirects!"), Args);
+			ValidationErrors.Emplace(NullValueError);
+		}
+	}
+
+	return Result;
+}
+#endif	// WITH_EDITOR
+
+const UInputMappingContext* UEnhancedInputPlatformData::GetContextRedirect(UInputMappingContext* InContext) const
+{
+	if (const TObjectPtr<const UInputMappingContext>* RedirectToIMC = MappingContextRedirects.Find(InContext))
+	{
+		return RedirectToIMC->Get();
+	}
+	return InContext;
+}
+
+//////////////////////////////////////////////////////////////////////////
+// UEnhancedInputPlatformSettings
+
+void UEnhancedInputPlatformSettings::PostLoad()
+{
+	Super::PostLoad();
+
+	InputDataClasses.Reset();
+	LoadInputDataClasses();
+}
+
+#if WITH_EDITOR
+void UEnhancedInputPlatformSettings::PostEditChangeProperty(FPropertyChangedEvent& PropertyChangedEvent)
+{
+	Super::PostEditChangeProperty(PropertyChangedEvent);
+
+	// Reset the cached data classes so that they get re-evaluated the next time we want to load them
+	// This ensures that we always have valid cached values during PIE sessions 
+	InputDataClasses.Reset();
+}
+#endif	// WITH_EDITOR
+
+void UEnhancedInputPlatformSettings::GetAllMappingContextRedirects(OUT TMap<TObjectPtr<const UInputMappingContext>, TObjectPtr<const UInputMappingContext>>& OutRedirects)
+{
+	ForEachInputData([&OutRedirects](const UEnhancedInputPlatformData& Data)
+	{
+		OutRedirects.Append(Data.GetMappingContextRedirects());
+	});
+}
+
+void UEnhancedInputPlatformSettings::ForEachInputData(TFunctionRef<void(const UEnhancedInputPlatformData&)> Predicate)
+{
+	LoadInputDataClasses();
+	
+	for (TSubclassOf<UEnhancedInputPlatformData> InputDataClass : InputDataClasses)
+	{
+		if (const UEnhancedInputPlatformData* Data = InputDataClass.GetDefaultObject())
+		{
+			Predicate(*Data);
+		}
+	}
+}
+
+void UEnhancedInputPlatformSettings::LoadInputDataClasses()
+{
+	if (InputData.Num() != InputDataClasses.Num())
+	{
+		for (TSoftClassPtr<UEnhancedInputPlatformData> InputDataPtr : InputData)
+		{
+			if (TSubclassOf<UEnhancedInputPlatformData> InputDataClass = InputDataPtr.LoadSynchronous())
+			{
+				InputDataClasses.Add(InputDataClass);
+			}
+		}
+	}
+}
+
+#undef LOCTEXT_NAMESPACE
\ No newline at end of file
diff --git a/Engine/Plugins/EnhancedInput/Source/EnhancedInput/Private/EnhancedInputSubsystemInterface.cpp b/Engine/Plugins/EnhancedInput/Source/EnhancedInput/Private/EnhancedInputSubsystemInterface.cpp
--- a/Engine/Plugins/EnhancedInput/Source/EnhancedInput/Private/EnhancedInputSubsystemInterface.cpp
+++ b/Engine/Plugins/EnhancedInput/Source/EnhancedInput/Private/EnhancedInputSubsystemInterface.cpp
@@ -12,6 +12,7 @@
 #include "InputTriggers.h"
 #include "UObject/UObjectIterator.h"
 #include "PlayerMappableInputConfig.h"
+#include "EnhancedInputPlatformSettings.h"
 
 /* Shared input subsystem functionality.
  * See EnhancedInputSubsystemInterfaceDebug.cpp for debug specific functionality.
@@ -546,9 +547,46 @@ void IEnhancedInputSubsystemInterface::RebuildControlMappings()
 	// Clear existing mappings, but retain the mapping array for later processing
 	TArray<FEnhancedActionKeyMapping> OldMappings(MoveTemp(PlayerInput->EnhancedActionMappings));
 	PlayerInput->ClearAllMappings();
+	AppliedContextRedirects.Reset();
 
 	// Order contexts by priority
 	TMap<TObjectPtr<const UInputMappingContext>, int32> OrderedInputContexts = PlayerInput->AppliedInputContexts;
+	
+	// Replace any mapping contexts that may have a redirect on this platform
+	if (UEnhancedInputPlatformSettings* PlatformSettings = UEnhancedInputPlatformSettings::Get())
+	{
+		TMap<TObjectPtr<const UInputMappingContext>, TObjectPtr<const UInputMappingContext>> ContextRedirects;
+		PlatformSettings->GetAllMappingContextRedirects(ContextRedirects);
+		for (const TPair<TObjectPtr<const UInputMappingContext>, TObjectPtr<const UInputMappingContext>> Pair : ContextRedirects)
+		{
+			if (Pair.Key.IsNull() || Pair.Value.IsNull())
+			{
+				UE_LOG(LogEnhancedInput, Error, TEXT("An invalid Mappping Context Redirect specified in '%s'"), PlatformSettings->GetConfigOverridePlatform());
+				continue;
+			}
+			
+			// Replace the existing IMC with the one that it should be redirected to on the PlayerInput 
+			if (const int32* ExistingIMCPriority = OrderedInputContexts.Find(Pair.Key))
+			{
+				OrderedInputContexts.Remove(Pair.Key);
+				OrderedInputContexts.Add(Pair.Value, *ExistingIMCPriority);
+				AppliedContextRedirects.Add(Pair);
+
+				// Optional logging that may be helpful for debugging purposes
+				if (PlatformSettings->ShouldLogMappingContextRedirects())
+				{
+					UE_LOG(
+						LogEnhancedInput,
+						Log,
+						TEXT("'%s' Redirecting Mapping Context '%s' -> '%s'"),
+						PlatformSettings->GetConfigOverridePlatform(), *Pair.Key->GetName(), *Pair.Value->GetName()
+					);
+				}
+			}
+		}
+	}
+	
+	// Order contexts by priority
 	OrderedInputContexts.ValueSort([](const int32& A, const int32& B) { return A > B; });
 
 	TSet<FKey> AppliedKeys;
diff --git a/Engine/Plugins/EnhancedInput/Source/EnhancedInput/Private/EnhancedInputSubsystemInterfaceDebug.cpp b/Engine/Plugins/EnhancedInput/Source/EnhancedInput/Private/EnhancedInputSubsystemInterfaceDebug.cpp
--- a/Engine/Plugins/EnhancedInput/Source/EnhancedInput/Private/EnhancedInputSubsystemInterfaceDebug.cpp
+++ b/Engine/Plugins/EnhancedInput/Source/EnhancedInput/Private/EnhancedInputSubsystemInterfaceDebug.cpp
@@ -13,7 +13,7 @@
 #include "InputModifiers.h"
 #include "InputTriggers.h"
 #include "ImageUtils.h"
-
+#include "EnhancedInputPlatformSettings.h"
 
 /* Shared input subsystem debug functionality.
  * See EnhancedInputSubsystemInterface.cpp for main functionality.
@@ -125,18 +125,32 @@ void IEnhancedInputSubsystemInterface::ShowDebugInfo(UCanvas* Canvas)
 
 		TMap<TObjectPtr<const UInputMappingContext>, int32> OrderedInputContexts = PlayerInput->AppliedInputContexts;
 		OrderedInputContexts.ValueSort([](const int32& A, const int32& B) { return A > B; });
-
+		
 		// Work through all input contexts, displaying active mappings, overridden mappings, etc.
 		for (const TPair<TObjectPtr<const UInputMappingContext>, int32>& ContextPair : OrderedInputContexts)
 		{
-			const UInputMappingContext* Context = ContextPair.Key.Get();
+			const UInputMappingContext* AppliedContext = ContextPair.Key.Get();
+
 			DisplayDebugManager.SetDrawColor(FColor::Yellow);
-			DisplayDebugManager.DrawString(FString::Printf(TEXT("  Context: %s"), *Context->GetFName().ToString()));
+			
+			// If this context was redirected via platform settings, then add some debug info about it here
+			if (const TObjectPtr<const UInputMappingContext>* RedirectedContextPtr = AppliedContextRedirects.Find(AppliedContext))
+			{
+				const UInputMappingContext* RedirectedContext = *RedirectedContextPtr;
+				DisplayDebugManager.DrawString(FString::Printf(TEXT("  Redirected Context: %s -> %s"), *AppliedContext->GetFName().ToString(), *RedirectedContext->GetFName().ToString()));
+
+				// Change the current context that is being used to display the correct mappings
+				AppliedContext = RedirectedContext;
+			}
+			else
+			{
+				DisplayDebugManager.DrawString(FString::Printf(TEXT("  Context: %s"), *AppliedContext->GetFName().ToString()));
+			}
 
 			// Build a table of mappings per action
 			TArray<const UInputAction*> OrderedActions;
 			TMap<const UInputAction*, TArray<FEnhancedActionKeyMapping>> ActionMappings;
-			for (const FEnhancedActionKeyMapping& Mapping : Context->GetMappings())
+			for (const FEnhancedActionKeyMapping& Mapping : AppliedContext->GetMappings())
 			{
 				if (Mapping.Action)
 				{
@@ -202,7 +216,7 @@ void IEnhancedInputSubsystemInterface::ShowDebugInfo(UCanvas* Canvas)
 					else if (!bHasChords && Mapping.Action->bConsumeInput)
 					{
 						// This mapping owns this key!
-						AppliedKeys.Emplace(Mapping.Key.GetFName(), Context->GetFName().ToString() + ":" + Action->GetFName().ToString());
+						AppliedKeys.Emplace(Mapping.Key.GetFName(), AppliedContext->GetFName().ToString() + ":" + Action->GetFName().ToString());
 					}
 					else if (bHasChords && KeyOwner)
 					{
diff --git a/Engine/Plugins/EnhancedInput/Source/EnhancedInput/Public/EnhancedInputDeveloperSettings.h b/Engine/Plugins/EnhancedInput/Source/EnhancedInput/Public/EnhancedInputDeveloperSettings.h
new file mode 100644
--- /dev/null
+++ b/Engine/Plugins/EnhancedInput/Source/EnhancedInput/Public/EnhancedInputDeveloperSettings.h
@@ -0,0 +1,25 @@
+// Copyright Epic Games, Inc. All Rights Reserved.
+
+#pragma once
+
+#include "CoreMinimal.h"
+#include "Engine/DeveloperSettingsBackedByCVars.h"
+#include "Engine/PlatformSettings.h"
+#include "EnhancedInputDeveloperSettings.generated.h"
+
+/** Developer settings for Enhanced Input */
+UCLASS(config = Input, defaultconfig)
+class ENHANCEDINPUT_API UEnhancedInputDeveloperSettings : public UDeveloperSettingsBackedByCVars
+{
+	GENERATED_BODY()
+public:
+	
+	UEnhancedInputDeveloperSettings(const FObjectInitializer& Initializer);
+
+	/**
+	 * Platform specific settings for Enhanced Input.
+	 * @see UEnhancedInputPlatformSettings
+	 */
+	UPROPERTY(EditAnywhere, Category = "Input")
+	FPerPlatformSettings PlatformSettings;
+};
\ No newline at end of file
diff --git a/Engine/Plugins/EnhancedInput/Source/EnhancedInput/Public/EnhancedInputPlatformSettings.h b/Engine/Plugins/EnhancedInput/Source/EnhancedInput/Public/EnhancedInputPlatformSettings.h
new file mode 100644
--- /dev/null
+++ b/Engine/Plugins/EnhancedInput/Source/EnhancedInput/Public/EnhancedInputPlatformSettings.h
@@ -0,0 +1,97 @@
+// Copyright Epic Games, Inc. All Rights Reserved.
+
+#pragma once
+
+#include "CoreMinimal.h"
+#include "Engine/PlatformSettings.h"
+#include "EnhancedInputPlatformSettings.generated.h"
+
+class UInputMappingContext;
+
+/**
+ * A base class that can be used to store platform specific data for Enhanced Input.
+ * 
+ * Make a subclass of this to add some additional options for per-platform settings 
+ */
+UCLASS(Abstract, Blueprintable, ClassGroup = Input, meta = (Category = "EnhancedInput"))
+class ENHANCEDINPUT_API UEnhancedInputPlatformData : public UObject
+{
+	GENERATED_BODY()
+
+public:
+
+	//~ UObject interface
+#if WITH_EDITOR
+	virtual EDataValidationResult IsDataValid(TArray<FText>& ValidationErrors) override;
+#endif // WITH_EDITOR
+	//~ End UObject interface
+	
+	/** Returns a pointer to the desired redirect mapping context. If there isn't one, then this returns InContext. */
+	UFUNCTION(BlueprintCallable, Category = "EnhancedInput")
+	const UInputMappingContext* GetContextRedirect(UInputMappingContext* InContext) const;
+
+	/** Return the map of all Context Redirects in this platform data. */
+	const TMap<TObjectPtr<const UInputMappingContext>, TObjectPtr<const UInputMappingContext>>& GetMappingContextRedirects() const { return MappingContextRedirects; }
+	
+protected:
+	
+	/**
+	 * Maps one Input Mapping Context to another. This can be used to replace
+	 * specific Input Mapping Contexts with another on a per-platform basis. 
+	 */
+	UPROPERTY(BlueprintReadOnly, EditDefaultsOnly, Category = "EnhancedInput")
+	TMap<TObjectPtr<const UInputMappingContext>, TObjectPtr<const UInputMappingContext>> MappingContextRedirects;
+};
+
+/** Settings for Enhanced Input that can be set per-platform */
+UCLASS(config = Input, defaultconfig)
+class ENHANCEDINPUT_API UEnhancedInputPlatformSettings : public UPlatformSettings
+{
+	GENERATED_BODY()
+
+public:
+
+	static UEnhancedInputPlatformSettings* Get()
+	{
+		return UPlatformSettingsManager::Get().GetSettingsForPlatform<UEnhancedInputPlatformSettings>();
+	}
+
+	//~ UObject interface
+	virtual void PostLoad() override;
+#if WITH_EDITOR
+	virtual void PostEditChangeProperty(struct FPropertyChangedEvent& PropertyChangedEvent) override;
+#endif
+	//~ End UObject interface
+
+	/** Get the current Enhanced Input platform data */
+	const TArray<TSoftClassPtr<UEnhancedInputPlatformData>>& GetInputData() const { return InputData; }
+
+	/** Populates the given map with all mapping context redirects in the current platform's InputData. */
+	void GetAllMappingContextRedirects(OUT TMap<TObjectPtr<const UInputMappingContext>, TObjectPtr<const UInputMappingContext>>& OutRedirects);
+
+	/** Iterates each valid UEnhancedInputPlatformData that is on this platform. */
+	void ForEachInputData(TFunctionRef<void(const UEnhancedInputPlatformData&)> Predicate);
+
+	/** Returns true if this platform has specified that it should log what mapping context redirects occur */
+	bool ShouldLogMappingContextRedirects() const { return bShouldLogMappingContextRedirects; }
+	
+protected:
+
+	/** Load the input data subclasses and cache them for later use in the InputDataClasses array. */
+	void LoadInputDataClasses();
+	
+	/** Input data that can be populated with Enhanced Input Platform Data blueprints */
+	UPROPERTY(config, EditAnywhere, Category = "Default")
+	TArray<TSoftClassPtr<UEnhancedInputPlatformData>> InputData;
+
+	/**
+	 * A transient array of the subclasses for the Enhanced Input Platform data. This will prevent us from
+	 * from having to load the class' default object during game time.
+	 */
+	UPROPERTY(Transient)
+	TArray<TSubclassOf<UEnhancedInputPlatformData>> InputDataClasses;
+
+	/** If true, then Enhanced Input will log which mapping contexts have been redirected when building the control mappings. */
+	UPROPERTY(Config, EditAnywhere, Category = "Default")
+	bool bShouldLogMappingContextRedirects = false;
+};
\ No newline at end of file
diff --git a/Engine/Plugins/EnhancedInput/Source/EnhancedInput/Public/EnhancedInputSubsystemInterface.h b/Engine/Plugins/EnhancedInput/Source/EnhancedInput/Public/EnhancedInputSubsystemInterface.h
--- a/Engine/Plugins/EnhancedInput/Source/EnhancedInput/Public/EnhancedInputSubsystemInterface.h
+++ b/Engine/Plugins/EnhancedInput/Source/EnhancedInput/Public/EnhancedInputSubsystemInterface.h
@@ -249,6 +249,12 @@ private:
 	/** A map of any player mapped keys to the key that they should redirect to instead */
 	TMap<FName, FKey> PlayerMappedSettings;
 
+	/**
+	 * A map of the currently applied mapping context redirects. This is populated in RebuildControlMappings
+	 * with any InputMappingContexts that have been redirected on this local player from the platform settings.
+	 */
+	TMap<TObjectPtr<const UInputMappingContext>, TObjectPtr<const UInputMappingContext>> AppliedContextRedirects;
+
 	EInputMappingRebuildType MappingRebuildPending = EInputMappingRebuildType::None;
 
 	/**
