Refactor of the Enhanced Input Project Settings. This moves the Trigger and Modifier default value customizations to the "InputCustomizations" header file, and out of the InputEditorModule. This allows us to simply add a class customization for the Developer Settings, and combine the two setting panels.

diff --git a/Engine/Plugins/EnhancedInput/Source/InputEditor/Private/InputCustomizations.cpp b/Engine/Plugins/EnhancedInput/Source/InputEditor/Private/InputCustomizations.cpp
--- a/Engine/Plugins/EnhancedInput/Source/InputEditor/Private/InputCustomizations.cpp
+++ b/Engine/Plugins/EnhancedInput/Source/InputEditor/Private/InputCustomizations.cpp
@@ -10,9 +10,17 @@
 #include "InputMappingContext.h"
 #include "KeyStructCustomization.h"
 #include "PropertyCustomizationHelpers.h"
+#include "EnhancedInputDeveloperSettings.h"
+#include "InputEditorModule.h"
+#include "AssetRegistryModule.h"
+#include "Blueprint/BlueprintSupport.h"
+#include "UObject/UObjectIterator.h"
 
 #define LOCTEXT_NAMESPACE "InputCustomization"
 
+//////////////////////////////////////////////////////////
+// FInputContextDetails
+
 TSharedRef<IDetailCustomization> FInputContextDetails::MakeInstance()
 {
 	return MakeShareable(new FInputContextDetails);
@@ -29,6 +37,8 @@ void FInputContextDetails::CustomizeDetails(class IDetailLayoutBuilder& DetailBu
 	MappingsDetailCategoryBuilder.AddCustomBuilder(ActionMappingsBuilder);
 }
 
+//////////////////////////////////////////////////////////
+// FEnhancedActionMappingCustomization
 
 void FEnhancedActionMappingCustomization::CustomizeHeader(TSharedRef<IPropertyHandle> PropertyHandle, FDetailWidgetRow& HeaderRow, IPropertyTypeCustomizationUtils& CustomizationUtils)
 {
@@ -74,4 +84,183 @@ void FEnhancedActionMappingCustomization::RemoveMappingButton_OnClick() const
 	}
 }
 
+//////////////////////////////////////////////////////////
+// FEnhancedInputDeveloperSettingsCustomization
+
+FEnhancedInputDeveloperSettingsCustomization::~FEnhancedInputDeveloperSettingsCustomization()
+{
+	// Unregister settings panel listeners
+	if (FAssetRegistryModule* AssetRegistryModule = FModuleManager::GetModulePtr<FAssetRegistryModule>("AssetRegistry"))
+	{
+		AssetRegistryModule->Get().OnAssetAdded().RemoveAll(this);
+		AssetRegistryModule->Get().OnAssetRemoved().RemoveAll(this);
+		AssetRegistryModule->Get().OnAssetRenamed().RemoveAll(this);
+	}
+	
+	CachedDetailBuilder = nullptr;
+}
+
+void FEnhancedInputDeveloperSettingsCustomization::CustomizeDetails(IDetailLayoutBuilder& DetailBuilder)
+{
+	// Note: The details view for the UEnhancedInputDeveloperSettings object will be on by default
+	// This 'EditCategory' call will ensure that the developer settings are displayed on top of the trigger/modifier default values
+	DetailBuilder.EditCategory(UEnhancedInputDeveloperSettings::StaticClass()->GetFName(), FText::GetEmpty(), ECategoryPriority::Important);
+
+	static const FName TriggerCategoryName = TEXT("Trigger Default Values");
+	static const FName ModifierCategoryName = TEXT("Modifier Default Values");
+
+	TArray<UObject*> ModifierCDOs = GatherClassDetailsCDOs(UInputModifier::StaticClass());
+	TArray<UObject*> TriggerCDOs = GatherClassDetailsCDOs(UInputTrigger::StaticClass());
+	
+	// Add The modifier/trigger defaults that are generated via CDO to the details builder
+	CustomizeCDOValues(DetailBuilder, ModifierCategoryName, ModifierCDOs);
+	CustomizeCDOValues(DetailBuilder, TriggerCategoryName, TriggerCDOs);
+	
+	// Support for updating blueprint based triggers and modifiers in the settings panel
+	FAssetRegistryModule& AssetRegistryModule = FModuleManager::LoadModuleChecked<FAssetRegistryModule>(FName("AssetRegistry"));
+	IAssetRegistry& AssetRegistry = AssetRegistryModule.Get();
+	if (!AssetRegistry.OnAssetAdded().IsBoundToObject(this))
+	{
+		AssetRegistry.OnAssetAdded().AddRaw(this, &FEnhancedInputDeveloperSettingsCustomization::OnAssetAdded);
+		AssetRegistry.OnAssetRemoved().AddRaw(this, &FEnhancedInputDeveloperSettingsCustomization::OnAssetRemoved);
+		AssetRegistry.OnAssetRenamed().AddRaw(this, &FEnhancedInputDeveloperSettingsCustomization::OnAssetRenamed);
+	}
+}
+
+void FEnhancedInputDeveloperSettingsCustomization::CustomizeCDOValues(IDetailLayoutBuilder& DetailBuilder, const FName CategoryName, const TArray<UObject*>& ObjectsToCustomize)
+{
+	IDetailCategoryBuilder& CategoryBuilder = DetailBuilder.EditCategory(CategoryName);
+
+	// All of the Objects in this array are the CDO, i.e. the class name starts with "Default__"
+	for (UObject* CDO : ObjectsToCustomize)
+	{
+		if (!ensure(CDO && CDO->IsTemplate()))
+		{
+			continue;
+		}
+
+		// Add the CDO as an external object reference to this customization
+		IDetailPropertyRow* Row = CategoryBuilder.AddExternalObjects({ CDO }, EPropertyLocation::Default, FAddPropertyParams().UniqueId(CDO->GetClass()->GetFName()));
+
+		if (Row)
+		{
+			// We need to add a custom "Name" widget here, otherwise all the categories will just say "Object"
+			Row->CustomWidget()
+			.NameContent()
+			[
+				SNew(STextBlock)
+					.Text(CDO->GetClass()->GetDisplayNameText())
+			];	
+		}
+	}
+}
+
+void FEnhancedInputDeveloperSettingsCustomization::CustomizeDetails(const TSharedPtr<IDetailLayoutBuilder>& DetailBuilder)
+{
+	CachedDetailBuilder = DetailBuilder;
+	CustomizeDetails(*DetailBuilder);
+}
+
+TArray<UObject*> FEnhancedInputDeveloperSettingsCustomization::GatherClassDetailsCDOs(UClass* Class)
+{
+	TArray<UObject*> CDOs;
+
+	// Search native classes
+	for (TObjectIterator<UClass> ClassIt; ClassIt; ++ClassIt)
+	{
+		if (!ClassIt->IsNative() || !ClassIt->IsChildOf(Class))
+		{
+			continue;
+		}
+
+		// Ignore abstract, hidedropdown, and deprecated.
+		if (ClassIt->HasAnyClassFlags(CLASS_Abstract | CLASS_HideDropDown | CLASS_Deprecated | CLASS_NewerVersionExists))
+		{
+			continue;
+		}
+
+		CDOs.AddUnique(ClassIt->GetDefaultObject());
+	}
+
+	// Search BPs via asset registry
+	FAssetRegistryModule& AssetRegistryModule = FModuleManager::LoadModuleChecked<FAssetRegistryModule>(FName("AssetRegistry"));
+	IAssetRegistry& AssetRegistry = AssetRegistryModule.Get();
+
+	FARFilter Filter;
+	Filter.ClassNames.Add(UBlueprint::StaticClass()->GetFName());
+	Filter.bRecursiveClasses = true;
+	TArray<FAssetData> BlueprintAssetData;
+	AssetRegistry.GetAssets(Filter, BlueprintAssetData);
+
+	for (FAssetData& Asset : BlueprintAssetData)
+	{
+		FAssetDataTagMapSharedView::FFindTagResult Result = Asset.TagsAndValues.FindTag(TEXT("NativeParentClass"));
+		if (Result.IsSet())
+		{
+			const FString ClassObjectPath = FPackageName::ExportTextPathToObjectPath(Result.GetValue());
+			const FString ClassName = FPackageName::ObjectPathToObjectName(ClassObjectPath);
+			if (UClass* ParentClass = FindObjectSafe<UClass>(ANY_PACKAGE, *ClassName, true))
+			{
+				if (ParentClass->IsChildOf(Class))
+				{
+					// TODO: Forcibly loading these assets could cause problems on projects with a large number of them.
+					UBlueprint* BP = CastChecked<UBlueprint>(Asset.GetAsset());
+					CDOs.AddUnique(BP->GeneratedClass->GetDefaultObject());
+				}
+			}
+		}
+	}
+
+	// Strip objects with no config stored properties
+	CDOs.RemoveAll([Class](UObject* Object) {
+		UClass* ObjectClass = Object->GetClass();
+		if (ObjectClass->GetMetaData(TEXT("NotInputConfigurable")).ToBool())
+		{
+			return true;
+		}
+		while (ObjectClass)
+		{
+			for (FProperty* Property : TFieldRange<FProperty>(ObjectClass, EFieldIteratorFlags::ExcludeSuper, EFieldIteratorFlags::ExcludeDeprecated))
+			{
+				if (Property->HasAnyPropertyFlags(CPF_Config))
+				{
+					return false;
+				}
+			}
+
+			// Stop searching at the base type. We don't care about configurable properties lower than that.
+			ObjectClass = ObjectClass != Class ? ObjectClass->GetSuperClass() : nullptr;
+		}
+		return true;
+	});
+
+	return CDOs;
+}
+
+void FEnhancedInputDeveloperSettingsCustomization::RebuildDetailsViewForAsset(const FAssetData& AssetData)
+{
+	// If the asset was a blueprint...
+	if (AssetData.AssetClass == UBlueprint::StaticClass()->GetFName())
+	{
+		// With a native parent class...
+		FAssetDataTagMapSharedView::FFindTagResult Result = AssetData.TagsAndValues.FindTag(FBlueprintTags::NativeParentClassPath);
+		if (Result.IsSet())
+		{
+			// And the base class is UInputModifier or UInputTrigger, then we should rebuild the details
+			const FString ClassObjectPath = FPackageName::ExportTextPathToObjectPath(Result.GetValue());
+			const FString ClassName = FPackageName::ObjectPathToObjectName(ClassObjectPath);
+			if (UClass* ParentClass = FindObjectSafe<UClass>(ANY_PACKAGE, *ClassName, true))
+			{
+				if (ParentClass == UInputModifier::StaticClass() || ParentClass == UInputTrigger::StaticClass())
+				{
+					if (IDetailLayoutBuilder* DetailBuilder = CachedDetailBuilder.Pin().Get())
+					{
+						DetailBuilder->ForceRefreshDetails();
+					}
+				}
+			}
+		}
+	}
+}
+
 #undef LOCTEXT_NAMESPACE
\ No newline at end of file
diff --git a/Engine/Plugins/EnhancedInput/Source/InputEditor/Private/InputCustomizations.h b/Engine/Plugins/EnhancedInput/Source/InputEditor/Private/InputCustomizations.h
--- a/Engine/Plugins/EnhancedInput/Source/InputEditor/Private/InputCustomizations.h
+++ b/Engine/Plugins/EnhancedInput/Source/InputEditor/Private/InputCustomizations.h
@@ -7,6 +7,8 @@
 #include "IDetailCustomNodeBuilder.h"
 #include "IPropertyTypeCustomization.h"
 
+class IDetailLayoutBuilder;
+
 class FInputContextDetails : public IDetailCustomization
 {
 public:
@@ -14,7 +16,7 @@ public:
 	static TSharedRef<IDetailCustomization> MakeInstance();
 
 	/** ILayoutDetails interface */
-	virtual void CustomizeDetails( class IDetailLayoutBuilder& DetailBuilder ) override;
+	virtual void CustomizeDetails( IDetailLayoutBuilder& DetailBuilder ) override;
 };
 
 class FEnhancedActionMappingCustomization : public IPropertyTypeCustomization
@@ -36,3 +38,54 @@ private:
 	TSharedPtr<IPropertyTypeCustomization> KeyStructInstance;
 	TSharedPtr<IPropertyHandle> MappingPropertyHandle;
 };
+
+/**
+ * Customization for UEnhancedInputDeveloperSettings.
+ *
+ * This will just make the normal details panel for UEnhancedInputDeveloperSettings, and then add all the default settings
+ * of Input Triggers and Input Modifiers by gather all the CDO's for them.
+ */
+class FEnhancedInputDeveloperSettingsCustomization : public IDetailCustomization
+{
+public:
+	//~ IDetailCustomization interface
+	static TSharedRef<IDetailCustomization> MakeInstance()
+	{
+		return MakeShareable(new FEnhancedInputDeveloperSettingsCustomization());
+	}
+
+	virtual ~FEnhancedInputDeveloperSettingsCustomization();
+	virtual void CustomizeDetails(IDetailLayoutBuilder& DetailBuilder) override;
+	virtual void CustomizeDetails(const TSharedPtr<IDetailLayoutBuilder>& DetailBuilder) override;
+	//~ End IDetailCustomization interface
+	
+private:
+
+	/** Gather all of the CDO's for the given class, Native and Blueprint. */
+	static TArray<UObject*> GatherClassDetailsCDOs(UClass* Class);
+
+	/**
+	 * Called when any Asset is added, removed, or renamed.
+	 *
+	 * This will rebuild the Modifier and Trigger CDO views to make sure that
+	 * any newly added Blueprint
+	 */
+	void RebuildDetailsViewForAsset(const FAssetData& AssetData);
+
+	/** Callbacks that are triggered from the Asset Registry. */
+	void OnAssetAdded(const FAssetData& AssetData) { RebuildDetailsViewForAsset(AssetData); }
+	void OnAssetRemoved(const FAssetData& AssetData) { RebuildDetailsViewForAsset(AssetData); }
+	void OnAssetRenamed(const FAssetData& AssetData, const FString&) { RebuildDetailsViewForAsset(AssetData); }
+	
+	/**
+	 * Create a new category on the DetailBuilder and add each object in the given array as an external reference.
+	 * 
+	 * @param DetailBuilder			The details builder that can be used to add categories
+	 * @param CategoryName			The name of the new category to add
+	 * @param ObjectsToCustomize	Array of CDO objects to customize and add as external references
+	 */
+	void CustomizeCDOValues(IDetailLayoutBuilder& DetailBuilder, const FName CategoryName, const TArray<UObject*>& ObjectsToCustomize);
+
+	// Cached details builder so that we can rebuild the details when a new BP asset is added
+	TWeakPtr<IDetailLayoutBuilder> CachedDetailBuilder;
+};
\ No newline at end of file
diff --git a/Engine/Plugins/EnhancedInput/Source/InputEditor/Private/InputEditorModule.cpp b/Engine/Plugins/EnhancedInput/Source/InputEditor/Private/InputEditorModule.cpp
--- a/Engine/Plugins/EnhancedInput/Source/InputEditor/Private/InputEditorModule.cpp
+++ b/Engine/Plugins/EnhancedInput/Source/InputEditor/Private/InputEditorModule.cpp
@@ -2,7 +2,6 @@
 
 #include "InputEditorModule.h"
 
-#include "AssetRegistryModule.h"
 #include "AssetTypeActions_Base.h"
 #include "EnhancedInputModule.h"
 #include "DetailCategoryBuilder.h"
@@ -15,9 +14,6 @@
 #include "PlayerMappableInputConfig.h"
 #include "InputCustomizations.h"
 #include "InputModifiers.h"
-#include "IAssetTools.h"
-#include "IAssetTypeActions.h"
-#include "IDetailsView.h"
 #include "ISettingsModule.h"
 #include "K2Node_EnhancedInputAction.h"
 #include "K2Node_GetInputActionValue.h"
@@ -26,69 +22,13 @@
 #include "PropertyEditorDelegates.h"
 #include "PropertyEditorModule.h"
 #include "SSettingsEditorCheckoutNotice.h"
-#include "TickableEditorObject.h"
 #include "Widgets/Layout/SScrollBox.h"
 #include "Widgets/Notifications/SNotificationList.h"
 #include "AssetTypeActions/AssetTypeActions_DataAsset.h"
+#include "EnhancedInputDeveloperSettings.h"
 
 #define LOCTEXT_NAMESPACE "InputEditor"
 
-class FInputEditorModule : public IModuleInterface, public FTickableEditorObject
-{
-public:
-
-	// IModuleInterface interface
-	virtual void StartupModule() override;
-	virtual void ShutdownModule() override;
-	// End IModuleInterface interface
-
-	// FTickableEditorObject interface
-	virtual void Tick(float DeltaTime) override;
-	virtual TStatId GetStatId() const override { RETURN_QUICK_DECLARE_CYCLE_STAT(FInputEditorModule, STATGROUP_Tickables); }
-	// End FTickableEditorObject interface
-
-	static EAssetTypeCategories::Type GetInputAssetsCategory() { return InputAssetsCategory; }
-
-private:
-	void RegisterAssetTypeActions(IAssetTools& AssetTools, TSharedRef<IAssetTypeActions> Action)
-	{
-		AssetTools.RegisterAssetTypeActions(Action);
-		CreatedAssetTypeActions.Add(Action);
-	}
-	void PostEngineInit();
-	TSharedRef<SWidget> CreateSettingsPanel();
-	void OnSettingChanged(const FPropertyChangedEvent& PropertyChangedEvent);
-
-	void RebuildDetailsViewForAsset(const FAssetData& AssetData, bool bIgnoreAsset);
-	void OnAssetAdded(const FAssetData& AssetData) { RebuildDetailsViewForAsset(AssetData, false); }
-	void OnAssetRemoved(const FAssetData& AssetData) { RebuildDetailsViewForAsset(AssetData, true); }
-	void OnAssetRenamed(const FAssetData& AssetData, const FString&) { RebuildDetailsViewForAsset(AssetData, true); }
-
-	template<typename T>
-	TSharedPtr<IDetailsView> AddClassDetailsView();
-
-	struct FClassDetailsView
-	{
-		FClassDetailsView() = default;
-		FClassDetailsView(UClass* InClass, TSharedPtr<IDetailsView>& InView) : Class(InClass), View(InView) {}
-		bool IsValid() const { return Class != nullptr; }
-
-		UClass* Class = nullptr;
-		TSharedPtr<IDetailsView> View;
-	};
-
-	FClassDetailsView FindClassDetailsViewForAsset(const FAssetData& AssetData);
-	TArray<UObject*> GatherClassDetailsCDOs(UClass* Class, const FAssetData* IgnoreAsset);
-
-	static EAssetTypeCategories::Type InputAssetsCategory;
-
-	TArray< TSharedPtr<IAssetTypeActions> > CreatedAssetTypeActions;
-
-	TMap<UClass*, TSharedPtr<IDetailsView>> DetailsViews;
-
-	TSharedPtr<SWidget> Panel;
-};
-
 EAssetTypeCategories::Type FInputEditorModule::InputAssetsCategory;
 
 IMPLEMENT_MODULE(FInputEditorModule, InputEditor)
@@ -190,258 +130,13 @@ public:
 	virtual UClass* GetSupportedClass() const override { return UPlayerMappableInputConfig::StaticClass(); }
 };
 
-
-
-void FInputEditorModule::OnSettingChanged(const FPropertyChangedEvent& PropertyChangedEvent)
-{
-	// TODO: Copy of SSettingsEditor::NotifyPostChange
-	if (PropertyChangedEvent.ChangeType != EPropertyChangeType::Interactive)
-	{
-		check(PropertyChangedEvent.GetNumObjectsBeingEdited() <= 1);
-		if(PropertyChangedEvent.GetNumObjectsBeingEdited() > 0)
-		{
-			UObject* ObjectBeingEdited = (UObject*)PropertyChangedEvent.GetObjectBeingEdited(0);
-
-			// Attempt to checkout the file automatically
-			if (!ObjectBeingEdited->GetClass()->HasAnyClassFlags(CLASS_DefaultConfig))
-			{
-				return;
-			}
-			FString RelativePath = ObjectBeingEdited->GetDefaultConfigFilename();
-
-			FString FullPath = FPaths::ConvertRelativePathToFull(RelativePath);
-
-			bool bIsNewFile = !FPlatformFileManager::Get().GetPlatformFile().FileExists(*FullPath);
-
-			if (!SettingsHelpers::CheckOutOrAddFile(FullPath))
-			{
-				SettingsHelpers::MakeWritable(FullPath);
-			}
-
-			// Determine if the Property is an Array or Array Element
-			bool bIsArrayOrArrayElement = PropertyChangedEvent.Property->IsA(FArrayProperty::StaticClass())
-				|| PropertyChangedEvent.Property->ArrayDim > 1
-				|| PropertyChangedEvent.Property->GetOwner<FArrayProperty>();
-
-			bool bIsSetOrSetElement = PropertyChangedEvent.Property->IsA(FSetProperty::StaticClass())
-				|| PropertyChangedEvent.Property->GetOwner<FSetProperty>();
-
-			bool bIsMapOrMapElement = PropertyChangedEvent.Property->IsA(FMapProperty::StaticClass())
-				|| PropertyChangedEvent.Property->GetOwner<FMapProperty>();
-
-			if (ObjectBeingEdited->GetClass()->HasAnyClassFlags(CLASS_DefaultConfig) && !bIsArrayOrArrayElement && !bIsSetOrSetElement && !bIsMapOrMapElement)
-			{
-				if(PropertyChangedEvent.Property->HasAnyPropertyFlags(CPF_Config | CPF_GlobalConfig))
-				{
-					ObjectBeingEdited->UpdateSinglePropertyInConfigFile(PropertyChangedEvent.Property, ObjectBeingEdited->GetDefaultConfigFilename());
-				}
-			}
-
-			if (bIsNewFile)
-			{
-				SettingsHelpers::CheckOutOrAddFile(FullPath);
-			}
-		}
-	}
-}
-
-template<typename T>
-class FPerCDOSettingsCustomization : public IDetailCustomization
-{
-	virtual void CustomizeDetails(IDetailLayoutBuilder& DetailBuilder) override
-	{
-		TArray<TWeakObjectPtr<UObject>> CustomizedObjects;
-		DetailBuilder.GetObjectsBeingCustomized(CustomizedObjects);
-
-		TWeakObjectPtr<UObject>& Object = CustomizedObjects[0];
-		if (Object.IsValid() && Object->HasAnyFlags(RF_ClassDefaultObject))
-		{
-			FString CategoryName = Object->GetClass()->GetName();
-			CategoryName.RemoveFromStart(T::StaticClass()->GetName());
-			CategoryName.RemoveFromEnd(TEXT("_C"));
-			IDetailCategoryBuilder& CategoryBuilder = DetailBuilder.EditCategory(Object->GetClass()->GetFName(), FText::FromString(FName::NameToDisplayString(CategoryName, false)));	// TODO: Category FName should be FullName for non-native objects
-
-			// TODO: Apply property categories as sub-categories?
-			UClass* BaseClass = Object->GetClass();
-			while (BaseClass)
-			{
-				for (FProperty* Property : TFieldRange<FProperty>(BaseClass, EFieldIteratorFlags::ExcludeSuper, EFieldIteratorFlags::ExcludeDeprecated))
-				{
-					if (Property->HasAnyPropertyFlags(CPF_Config))
-					{
-						CategoryBuilder.AddProperty(Property->GetFName(), BaseClass);
-					}
-				}
-
-				BaseClass = BaseClass != T::StaticClass() ? BaseClass->GetSuperClass() : nullptr; // Stop searching at the base type. We don't care about configurable properties lower than that.
-			}
-		}
-	}
-};
-
-
-template<typename T>
-TSharedPtr<IDetailsView> FInputEditorModule::AddClassDetailsView()
-{
-	FPropertyEditorModule& PropertyEditorModule = FModuleManager::GetModuleChecked<FPropertyEditorModule>("PropertyEditor");
-
-	FDetailsViewArgs DetailsViewArgs;
-	DetailsViewArgs.NameAreaSettings = FDetailsViewArgs::HideNameArea;
-	DetailsViewArgs.bHideSelectionTip = true;
-	DetailsViewArgs.bAllowMultipleTopLevelObjects = true;
-	DetailsViewArgs.bShowOptions = false;
-	DetailsViewArgs.bShowPropertyMatrixButton = false;
-
-	TSharedPtr<IDetailsView> DetailsView = PropertyEditorModule.CreateDetailView(DetailsViewArgs);
-	DetailsView->OnFinishedChangingProperties().AddRaw(this, &FInputEditorModule::OnSettingChanged);
-	DetailsView->RegisterInstancedCustomPropertyLayout(T::StaticClass(), FOnGetDetailCustomizationInstance::CreateLambda(&MakeShared<FPerCDOSettingsCustomization<T>>));
-
-	// Init CDOs for view
-	DetailsView->SetObjects(GatherClassDetailsCDOs(T::StaticClass(), nullptr));
-
-	return DetailsViews.Add(T::StaticClass(), DetailsView);
-}
-
-TArray<UObject*> FInputEditorModule::GatherClassDetailsCDOs(UClass* Class, const FAssetData* IgnoreAsset)
-{
-	TArray<UObject*> CDOs;
-
-	// Search native classes
-	for (TObjectIterator<UClass> ClassIt; ClassIt; ++ClassIt)
-	{
-		if (!ClassIt->IsNative() || !ClassIt->IsChildOf(Class))
-		{
-			continue;
-		}
-
-		// Ignore abstract, hidedropdown, and deprecated.
-		if (ClassIt->HasAnyClassFlags(CLASS_Abstract | CLASS_HideDropDown | CLASS_Deprecated | CLASS_NewerVersionExists))
-		{
-			continue;
-		}
-
-		CDOs.AddUnique(ClassIt->GetDefaultObject());
-	}
-
-	// Search BPs via asset registry
-	FAssetRegistryModule& AssetRegistryModule = FModuleManager::LoadModuleChecked<FAssetRegistryModule>(FName("AssetRegistry"));
-	IAssetRegistry& AssetRegistry = AssetRegistryModule.Get();
-
-	FARFilter Filter;
-	Filter.ClassNames.Add(UBlueprint::StaticClass()->GetFName());
-	Filter.bRecursiveClasses = true;
-	TArray<FAssetData> BlueprintAssetData;
-	AssetRegistry.GetAssets(Filter, BlueprintAssetData);
-
-	for (FAssetData& Asset : BlueprintAssetData)
-	{
-		if (IgnoreAsset && Asset == *IgnoreAsset)
-		{
-			continue;
-		}
-
-		FAssetDataTagMapSharedView::FFindTagResult Result = Asset.TagsAndValues.FindTag(TEXT("NativeParentClass"));
-		if (Result.IsSet())
-		{
-			const FString ClassObjectPath = FPackageName::ExportTextPathToObjectPath(Result.GetValue());
-			const FString ClassName = FPackageName::ObjectPathToObjectName(ClassObjectPath);
-			if (UClass* ParentClass = FindObjectSafe<UClass>(ANY_PACKAGE, *ClassName, true))
-			{
-				if (ParentClass->IsChildOf(Class))
-				{
-					// TODO: Forcibly loading these assets could cause problems on projects with a large number of them.
-					UBlueprint* BP = CastChecked<UBlueprint>(Asset.GetAsset());
-					CDOs.AddUnique(BP->GeneratedClass->GetDefaultObject());
-				}
-			}
-		}
-	}
-
-	// Strip objects with no config stored properties
-	CDOs.RemoveAll([Class](UObject* Object) {
-		UClass* ObjectClass = Object->GetClass();
-		if (ObjectClass->GetMetaData(TEXT("NotInputConfigurable")).ToBool())
-		{
-			return true;
-		}
-		while (ObjectClass)
-		{
-			for (FProperty* Property : TFieldRange<FProperty>(ObjectClass, EFieldIteratorFlags::ExcludeSuper, EFieldIteratorFlags::ExcludeDeprecated))
-			{
-				if (Property->HasAnyPropertyFlags(CPF_Config))
-				{
-					return false;
-				}
-			}
-
-			ObjectClass = ObjectClass != Class ? ObjectClass->GetSuperClass() : nullptr; // Stop searching at the base type. We don't care about configurable properties lower than that.
-		}
-		return true;
-	});
-
-	return CDOs;
-}
-
-void FInputEditorModule::PostEngineInit()
-{
-	// Register input settings
-	ISettingsModule* SettingsModule = FModuleManager::GetModulePtr<ISettingsModule>("Settings");
-	if (SettingsModule && FSlateApplication::IsInitialized())
-	{
-		SettingsModule->RegisterSettings("Project", "Plugins", "EnhancedInput",
-			LOCTEXT("EnhancedInputSettingsName", "Enhanced Input"),
-			LOCTEXT("EnhancedInputSettingsDescription", "Modify defaults for configurable triggers and modifiers."),
-			CreateSettingsPanel()
-		);
-	}
-}
-
-TSharedRef<SWidget> FInputEditorModule::CreateSettingsPanel()
-{
-	TSharedPtr<IDetailsView> TriggerDetailsView = AddClassDetailsView<UInputTrigger>();
-	TSharedPtr<IDetailsView> ModifierDetailsView = AddClassDetailsView<UInputModifier>();
-
-	FSlateFontInfo HeaderFont = FEditorStyle::GetFontStyle("DetailsView.CategoryFontStyle");
-	HeaderFont.Size = 32;
-
-	SAssignNew(Panel, SScrollBox)
-	+SScrollBox::Slot()
-	.Padding(0.f, 4.f, 0.f, 4.f)
-	[
-		SNew(STextBlock)
-		.Font(HeaderFont)
-		.Text(LOCTEXT("EngineInputSettingsTriggers", "Trigger Defaults"))
-	]
-
-	+SScrollBox::Slot()
-	.Padding(0.f, 4.f, 0.f, 0.f)
-	[
-		TriggerDetailsView.ToSharedRef()
-	]
-
-	+SScrollBox::Slot()
-	.Padding(0.f, 12.f, 0.f, 4.f)
-	[
-		SNew(STextBlock)
-		.Font(HeaderFont)
-		.Text(LOCTEXT("EngineInputSettingsModifiers", "Modifier Defaults"))
-	]
-
-	+SScrollBox::Slot()
-	.Padding(0.f, 8.f, 0.f, 0.f)
-	[
-		ModifierDetailsView.ToSharedRef()
-	];
-
-	return Panel.ToSharedRef();
-}
-
 void FInputEditorModule::StartupModule()
 {
 	// Register customizations
 	FPropertyEditorModule& PropertyModule = FModuleManager::GetModuleChecked<FPropertyEditorModule>("PropertyEditor");
 	PropertyModule.RegisterCustomClassLayout("InputMappingContext", FOnGetDetailCustomizationInstance::CreateStatic(&FInputContextDetails::MakeInstance));
 	PropertyModule.RegisterCustomPropertyTypeLayout("EnhancedActionKeyMapping", FOnGetPropertyTypeCustomizationInstance::CreateStatic(&FEnhancedActionMappingCustomization::MakeInstance));
+	PropertyModule.RegisterCustomClassLayout(UEnhancedInputDeveloperSettings::StaticClass()->GetFName(), FOnGetDetailCustomizationInstance::CreateStatic(&FEnhancedInputDeveloperSettingsCustomization::MakeInstance));
 	PropertyModule.NotifyCustomizationModuleChanged();
 
 	// Register input assets
@@ -455,28 +150,10 @@ void FInputEditorModule::StartupModule()
 		//RegisterAssetTypeActions(AssetTools, MakeShareable(new FAssetTypeActions_InputTrigger));
 		//RegisterAssetTypeActions(AssetTools, MakeShareable(new FAssetTypeActions_InputModifier));
 	}
-
-	// Support for updating blueprint based triggers and modifiers in the settings panel
-	FAssetRegistryModule& AssetRegistryModule = FModuleManager::LoadModuleChecked<FAssetRegistryModule>(FName("AssetRegistry"));
-	IAssetRegistry& AssetRegistry = AssetRegistryModule.Get();
-	AssetRegistry.OnAssetAdded().AddRaw(this, &FInputEditorModule::OnAssetAdded);
-	AssetRegistry.OnAssetRemoved().AddRaw(this, &FInputEditorModule::OnAssetRemoved);
-	AssetRegistry.OnAssetRenamed().AddRaw(this, &FInputEditorModule::OnAssetRenamed);
-	// TODO: Update settings whenever a config variable is added/removed to an asset
-
-	FCoreDelegates::OnPostEngineInit.AddRaw(this, &FInputEditorModule::PostEngineInit);
 }
 
 void FInputEditorModule::ShutdownModule()
 {
-	// Unregister settings panel listeners
-	if (FAssetRegistryModule* AssetRegistryModule = FModuleManager::GetModulePtr<FAssetRegistryModule>("AssetRegistry"))
-	{
-		AssetRegistryModule->Get().OnAssetAdded().RemoveAll(this);
-		AssetRegistryModule->Get().OnAssetRemoved().RemoveAll(this);
-		AssetRegistryModule->Get().OnAssetRenamed().RemoveAll(this);
-	}
-
 	// Unregister input assets
 	if (FAssetToolsModule* AssetToolsModule = FModuleManager::GetModulePtr<FAssetToolsModule>("AssetTools"))
 	{
@@ -497,40 +174,8 @@ void FInputEditorModule::ShutdownModule()
 	FPropertyEditorModule& PropertyModule = FModuleManager::GetModuleChecked<FPropertyEditorModule>("PropertyEditor");
 	PropertyModule.UnregisterCustomClassLayout("InputContext");
 	PropertyModule.UnregisterCustomPropertyTypeLayout("EnhancedActionKeyMapping");
+	PropertyModule.UnregisterCustomClassLayout(UEnhancedInputDeveloperSettings::StaticClass()->GetFName());
 	PropertyModule.NotifyCustomizationModuleChanged();
-
-	FCoreDelegates::OnPostEngineInit.RemoveAll(this);
-}
-
-FInputEditorModule::FClassDetailsView FInputEditorModule::FindClassDetailsViewForAsset(const FAssetData& AssetData)
-{
-	if (AssetData.AssetClass == UBlueprint::StaticClass()->GetFName())
-	{
-		FAssetDataTagMapSharedView::FFindTagResult Result = AssetData.TagsAndValues.FindTag(TEXT("NativeParentClass"));
-		if (Result.IsSet())
-		{
-			const FString ClassObjectPath = FPackageName::ExportTextPathToObjectPath(Result.GetValue());
-			const FString ClassName = FPackageName::ObjectPathToObjectName(ClassObjectPath);
-			if (UClass* ParentClass = FindObjectSafe<UClass>(ANY_PACKAGE, *ClassName, true))
-			{
-				if (TSharedPtr<IDetailsView>* DetailsView = DetailsViews.Find(ParentClass))
-				{
-					return FClassDetailsView(ParentClass, *DetailsView);
-				}
-			}
-		}
-	}
-
-	return FClassDetailsView();
-}
-
-void FInputEditorModule::RebuildDetailsViewForAsset(const FAssetData& AssetData, bool bIgnoreAsset)
-{
-	FClassDetailsView CDV = FindClassDetailsViewForAsset(AssetData);
-	if (CDV.IsValid())
-	{
-		CDV.View->SetObjects(GatherClassDetailsCDOs(CDV.Class, bIgnoreAsset ? &AssetData : nullptr));
-	}
 }
 
 void FInputEditorModule::Tick(float DeltaTime)
diff --git a/Engine/Plugins/EnhancedInput/Source/InputEditor/Public/InputEditorModule.h b/Engine/Plugins/EnhancedInput/Source/InputEditor/Public/InputEditorModule.h
--- a/Engine/Plugins/EnhancedInput/Source/InputEditor/Public/InputEditorModule.h
+++ b/Engine/Plugins/EnhancedInput/Source/InputEditor/Public/InputEditorModule.h
@@ -4,11 +4,46 @@
 
 #include "CoreMinimal.h"
 #include "Factories/Factory.h"
+#include "AssetTypeCategories.h"
+#include "IAssetTools.h"
+#include "IAssetTypeActions.h"
+#include "TickableEditorObject.h"
+#include "IDetailsView.h"
 
 #include "InputEditorModule.generated.h"
 
-// Module is not publicly exposed
+////////////////////////////////////////////////////////////////////
+// FInputEditorModule
 
+class FInputEditorModule : public IModuleInterface, public FTickableEditorObject
+{
+public:
+
+	// IModuleInterface interface
+	virtual void StartupModule() override;
+	virtual void ShutdownModule() override;
+	// End IModuleInterface interface
+
+	// FTickableEditorObject interface
+	virtual void Tick(float DeltaTime) override;
+	virtual TStatId GetStatId() const override { RETURN_QUICK_DECLARE_CYCLE_STAT(FInputEditorModule, STATGROUP_Tickables); }
+	// End FTickableEditorObject interface
+
+	static EAssetTypeCategories::Type GetInputAssetsCategory() { return InputAssetsCategory; }
+	
+private:
+	void RegisterAssetTypeActions(IAssetTools& AssetTools, TSharedRef<IAssetTypeActions> Action)
+	{
+		AssetTools.RegisterAssetTypeActions(Action);
+		CreatedAssetTypeActions.Add(Action);
+	}
+
+	static EAssetTypeCategories::Type InputAssetsCategory;
+	
+	TArray<TSharedPtr<IAssetTypeActions>> CreatedAssetTypeActions;
+};
+
+////////////////////////////////////////////////////////////////////
 // Asset factories
 
 UCLASS()
