Fix a linux ASan error when shutting down the InputEditorModule. We can use ::StaticClass()->GetName because the UObject may be uninitalized already

#rnx

diff --git a/Engine/Plugins/EnhancedInput/Source/InputEditor/Private/InputEditorModule.cpp b/Engine/Plugins/EnhancedInput/Source/InputEditor/Private/InputEditorModule.cpp
--- a/Engine/Plugins/EnhancedInput/Source/InputEditor/Private/InputEditorModule.cpp
+++ b/Engine/Plugins/EnhancedInput/Source/InputEditor/Private/InputEditorModule.cpp
@@ -174,7 +174,7 @@ void FInputEditorModule::ShutdownModule()
 	FPropertyEditorModule& PropertyModule = FModuleManager::GetModuleChecked<FPropertyEditorModule>("PropertyEditor");
 	PropertyModule.UnregisterCustomClassLayout("InputContext");
 	PropertyModule.UnregisterCustomPropertyTypeLayout("EnhancedActionKeyMapping");
-	PropertyModule.UnregisterCustomClassLayout(UEnhancedInputDeveloperSettings::StaticClass()->GetFName());
+	PropertyModule.UnregisterCustomClassLayout("EnhancedInputDeveloperSettings");
 	PropertyModule.NotifyCustomizationModuleChanged();
 }
 
