Set the display name on Player Mappable Input Configs correctly. They were previously "Bindable", but that naming isn't consistent

#rnx

diff --git a/Engine/Plugins/EnhancedInput/Source/EnhancedInput/Public/PlayerMappableInputConfig.h b/Engine/Plugins/EnhancedInput/Source/EnhancedInput/Public/PlayerMappableInputConfig.h
--- a/Engine/Plugins/EnhancedInput/Source/EnhancedInput/Public/PlayerMappableInputConfig.h
+++ b/Engine/Plugins/EnhancedInput/Source/EnhancedInput/Public/PlayerMappableInputConfig.h
@@ -16,9 +16,9 @@ class UInputMappingContext;
  * the default mappings for your player to start with in your game. It provides an easy way to get only the player
  * mappable key actions, and it can be used to add multiple UInputMappingContext's at once to the player.
  *
- * Populate this data asset with Input Mapping Contexts that have player bindable actions in them. 
+ * Populate this data asset with Input Mapping Contexts that have player mappable actions in them. 
  */
-UCLASS(BlueprintType, Meta = (DisplayName = "Player Bindable Input Config", ShortTooltip = "Data asset used to define a set of player mappable controller/keyboard mappings."))
+UCLASS(BlueprintType, Meta = (DisplayName = "Player Mappable Input Config", ShortTooltip = "Data asset used to define a set of player mappable controller/keyboard mappings."))
 class ENHANCEDINPUT_API UPlayerMappableInputConfig : public UPrimaryDataAsset
 {
 	GENERATED_BODY()
