Add a missing 'ENHANCEDINPUT_API' to the Enhanced Input editor subsystem so that it can be used outside of the module

#rnx

diff --git a/Engine/Plugins/EnhancedInput/Source/InputEditor/Public/EnhancedInputEditorSubsystem.h b/Engine/Plugins/EnhancedInput/Source/InputEditor/Public/EnhancedInputEditorSubsystem.h
--- a/Engine/Plugins/EnhancedInput/Source/InputEditor/Public/EnhancedInputEditorSubsystem.h
+++ b/Engine/Plugins/EnhancedInput/Source/InputEditor/Public/EnhancedInputEditorSubsystem.h
@@ -24,7 +24,7 @@ class FEnhancedInputEditorProcessor;
  * input.
  */
 UCLASS()
-class UEnhancedInputEditorSubsystem : public UEditorSubsystem, public IEnhancedInputSubsystemInterface, public FTickableGameObject
+class INPUTEDITOR_API UEnhancedInputEditorSubsystem : public UEditorSubsystem, public IEnhancedInputSubsystemInterface, public FTickableGameObject
 {
 	GENERATED_BODY()
 
