Create accessors on the Enhaced Player input for the currently applied mapping contexts and mappings.

These functions will make it easier for people to extend this class, while still keeping the AppliedInputContext and EnhancedActionMappings properties const and non-mutable.

diff --git a/Engine/Plugins/EnhancedInput/Source/EnhancedInput/Public/EnhancedPlayerInput.h b/Engine/Plugins/EnhancedInput/Source/EnhancedInput/Public/EnhancedPlayerInput.h
--- a/Engine/Plugins/EnhancedInput/Source/EnhancedInput/Public/EnhancedPlayerInput.h
+++ b/Engine/Plugins/EnhancedInput/Source/EnhancedInput/Public/EnhancedPlayerInput.h
@@ -53,6 +53,13 @@ protected:
 	
 	// Causes key to be consumed if it is affecting an action.
 	virtual bool IsKeyHandledByAction(FKey Key) const override;
+	
+	/** Note: Source reference only. Use EnhancedActionMappings for the actual mappings (with properly instanced triggers/modifiers) */
+	const TMap<TObjectPtr<const UInputMappingContext>, int32>& GetAppliedInputContexts() const { return AppliedInputContexts; }
+
+	/** This player's version of the Action Mappings */
+	const TArray<FEnhancedActionKeyMapping>& GetEnhancedActionMappings() const { return EnhancedActionMappings; }
+	
 
 private:
 
