First pass at a "showdebug Devices" command that will display debug info about all current platform users and their connected input devices. This will display some basic information on the screen that will be really useful while testing all this on different platforms.

Note: For now this lives with Enhanced Input, but should probably go somewhere else as it doesn't actually depend on anything related to Enhanced Input. Ideally I think it should live in InputCore, but it cant because InputCore cannot depend on the "Engine" module which is where all the debug helpers are.

diff --git a/Engine/Plugins/EnhancedInput/Source/EnhancedInput/EnhancedInput.Build.cs b/Engine/Plugins/EnhancedInput/Source/EnhancedInput/EnhancedInput.Build.cs
--- a/Engine/Plugins/EnhancedInput/Source/EnhancedInput/EnhancedInput.Build.cs
+++ b/Engine/Plugins/EnhancedInput/Source/EnhancedInput/EnhancedInput.Build.cs
@@ -14,7 +14,8 @@ namespace UnrealBuildTool.Rules
 
             PrivateDependencyModuleNames.AddRange(
 				new string[] {
-                    "Core",
+                    "ApplicationCore",
+					"Core",
 					"CoreUObject",
                     "Engine",
 					"InputCore",
diff --git a/Engine/Plugins/EnhancedInput/Source/EnhancedInput/Private/EnhancedInputModule.cpp b/Engine/Plugins/EnhancedInput/Source/EnhancedInput/Private/EnhancedInputModule.cpp
--- a/Engine/Plugins/EnhancedInput/Source/EnhancedInput/Private/EnhancedInputModule.cpp
+++ b/Engine/Plugins/EnhancedInput/Source/EnhancedInput/Private/EnhancedInputModule.cpp
@@ -301,20 +301,32 @@ void FEnhancedInputModule::Tick(float DeltaTime)
 void FEnhancedInputModule::OnShowDebugInfo(AHUD* HUD, UCanvas* Canvas, const FDebugDisplayInfo& DisplayInfo, float& YL, float& YPos)
 {
 	static const FName NAME_EnhancedInput("EnhancedInput");
-	if (Canvas && HUD->ShouldDisplayDebug(NAME_EnhancedInput))
+	static const FName NAME_PlatformDevices("Devices");
+	if (Canvas)
 	{
-		FDisplayDebugManager& DisplayDebugManager = Canvas->DisplayDebugManager;
-		DisplayDebugManager.SetFont(GEngine->GetSmallFont());
-		DisplayDebugManager.SetDrawColor(FColor::Yellow);
-		DisplayDebugManager.DrawString(TEXT("ENHANCED INPUT"));
-
-		// TODO: Support paging through subsystems one at a time (via console? key press?)
-
-		// Show first player only for now
-		TObjectIterator<UEnhancedInputLocalPlayerSubsystem> FirstPlayer;
-		if (FirstPlayer)
+		if (HUD->ShouldDisplayDebug(NAME_EnhancedInput))
+		{
+			FDisplayDebugManager& DisplayDebugManager = Canvas->DisplayDebugManager;
+			DisplayDebugManager.SetFont(GEngine->GetSmallFont());
+			DisplayDebugManager.SetDrawColor(FColor::Yellow);
+			DisplayDebugManager.DrawString(TEXT("ENHANCED INPUT"));
+			
+			// TODO: Support paging through subsystems one at a time (via console? key press?)
+
+			// Show first player only for now
+			TObjectIterator<UEnhancedInputLocalPlayerSubsystem> FirstPlayer;
+			if (FirstPlayer)
+			{
+				FirstPlayer->ShowDebugInfo(Canvas);
+			}
+		}
+		
+		if (HUD->ShouldDisplayDebug(NAME_PlatformDevices))
 		{
-			FirstPlayer->ShowDebugInfo(Canvas);
+			UEnhancedInputLibrary::ForEachSubsystem([Canvas](IEnhancedInputSubsystemInterface* Subsystem)
+			{
+				Subsystem->ShowPlatformInputDebugInfo(Canvas);
+			});
 		}
 	}
 	else
diff --git a/Engine/Plugins/EnhancedInput/Source/EnhancedInput/Private/EnhancedInputSubsystemInterfaceDebug.cpp b/Engine/Plugins/EnhancedInput/Source/EnhancedInput/Private/EnhancedInputSubsystemInterfaceDebug.cpp
--- a/Engine/Plugins/EnhancedInput/Source/EnhancedInput/Private/EnhancedInputSubsystemInterfaceDebug.cpp
+++ b/Engine/Plugins/EnhancedInput/Source/EnhancedInput/Private/EnhancedInputSubsystemInterfaceDebug.cpp
@@ -14,6 +14,7 @@
 #include "InputTriggers.h"
 #include "ImageUtils.h"
 #include "EnhancedInputPlatformSettings.h"
+#include "GenericPlatform/GenericPlatformInputDeviceMapper.h"
 
 /* Shared input subsystem debug functionality.
  * See EnhancedInputSubsystemInterface.cpp for main functionality.
@@ -271,6 +272,79 @@ void IEnhancedInputSubsystemInterface::ShowDebugInfo(UCanvas* Canvas)
 	}
 }
 
+void IEnhancedInputSubsystemInterface::ShowPlatformInputDebugInfo(class UCanvas* Canvas)
+{
+	if (!Canvas)
+	{
+		return;
+	}
+	FDisplayDebugManager& DisplayDebugManager = Canvas->DisplayDebugManager;
+
+	// Show title
+	DisplayDebugManager.SetFont(GEngine->GetMediumFont());
+	DisplayDebugManager.SetDrawColor(FColor::Orange);
+	DisplayDebugManager.DrawString(TEXT("\n\nPlatform Input Devices"));
+
+	// Gather all current platform users
+	IPlatformInputDeviceMapper& DeviceMapper = IPlatformInputDeviceMapper::Get();
+	TArray<FPlatformUserId> PlatformUsers;
+	const int32 NumUsers = DeviceMapper.GetAllActiveUsers(PlatformUsers);
+	
+	DisplayDebugManager.SetDrawColor(FColor::White);
+	DisplayDebugManager.DrawString(FString::Printf(TEXT("Total Platform Users: %d\n"), NumUsers));
+
+	// Show all the connected input devices for each available platform user
+	// TODO: We could make two commands, one to display all platform users and one to only show this local player's data
+	TArray<FInputDeviceId> UserDeviceIds;
+	for (FPlatformUserId& PlatUser : PlatformUsers)
+	{
+		const int32 NumDevices = DeviceMapper.GetAllInputDevicesForUser(PlatUser, UserDeviceIds);
+		FString UserDebugInfo = FString::Printf(TEXT("Platform User ID '%d' has '%d' input devices"), PlatUser.GetInternalId(), NumDevices);
+		// Show in red if there are no devices, that is a problem!
+		DisplayDebugManager.SetDrawColor(NumDevices > 0 ? FColor::White : FColor::Red);
+		DisplayDebugManager.DrawString(UserDebugInfo);
+
+		// Display the debug info about each input device
+		for (FInputDeviceId& Device : UserDeviceIds)
+		{
+			static const float InputDeviceXOffset = 10.0f;
+			// Display the input device ID
+			FString DeviceIdString = FString::Printf(TEXT("Device ID: %d"), Device.GetId());
+			
+			DisplayDebugManager.SetDrawColor(FColor::White);
+			DisplayDebugManager.DrawString(DeviceIdString, InputDeviceXOffset);
+			
+			// Display the connection state
+			FString ConnectionStateString;
+			FColor ConnectionStateColor = FColor::White;
+			const EInputDeviceConnectionState DeviceState = DeviceMapper.GetInputDeviceConnectionState(Device);
+			
+			switch(DeviceState)
+			{
+				case EInputDeviceConnectionState::Invalid:
+					ConnectionStateString = TEXT("INVALID");
+					ConnectionStateColor = FColor::Red;
+					break;
+				case EInputDeviceConnectionState::Unknown:
+					ConnectionStateString = TEXT("Unknown");
+					ConnectionStateColor = FColor::Orange;
+					break;
+				case EInputDeviceConnectionState::Disconnected:
+					ConnectionStateString = TEXT("Disconnected");
+					ConnectionStateColor = FColor::Silver;
+					break;
+				case EInputDeviceConnectionState::Connected:
+					ConnectionStateString = TEXT("Connected");
+					ConnectionStateColor = FColor::Green;
+					break;
+			}
+			
+			DisplayDebugManager.SetDrawColor(ConnectionStateColor);
+			DisplayDebugManager.DrawString(ConnectionStateString, InputDeviceXOffset);
+		}
+	}
+}
+
 void IEnhancedInputSubsystemInterface::ShowDebugActionModifiers(UCanvas* Canvas, const UInputAction* Action)
 {
 	FDisplayDebugManager& DisplayDebugManager = Canvas->DisplayDebugManager;
diff --git a/Engine/Plugins/EnhancedInput/Source/EnhancedInput/Public/EnhancedInputSubsystemInterface.h b/Engine/Plugins/EnhancedInput/Source/EnhancedInput/Public/EnhancedInputSubsystemInterface.h
--- a/Engine/Plugins/EnhancedInput/Source/EnhancedInput/Public/EnhancedInputSubsystemInterface.h
+++ b/Engine/Plugins/EnhancedInput/Source/EnhancedInput/Public/EnhancedInputSubsystemInterface.h
@@ -269,6 +269,12 @@ private:
 
 	// Debug visualization implemented in EnhancedInputSubsystemsDebug.cpp
 	void ShowDebugInfo(class UCanvas* Canvas);
+	
+	// Debug visualization of any platform input devices (IPlatformInputDeviceMapper)
+	// TODO: We should get a better place to put this instead of the Enhanced Input plugin.
+	// It can't go in InputCore because that doesn't depend on Engine for debug drawing
+	void ShowPlatformInputDebugInfo(class UCanvas* Canvas);
+	
 	void ShowDebugActionModifiers(UCanvas* Canvas, const UInputAction* Action);
 	static void PurgeDebugVisualizations();
 };
\ No newline at end of file
