[Enhanced Input] Deprecated UInputMappingContext::UnmapAction because the name is a bit confusing. Create a new "UnmapAllKeysFromAction" function instead to be more explicit.

diff --git a/Engine/Plugins/EnhancedInput/Config/BaseEngine.ini b/Engine/Plugins/EnhancedInput/Config/BaseEngine.ini
new file mode 100644
--- /dev/null
+++ b/Engine/Plugins/EnhancedInput/Config/BaseEngine.ini
@@ -0,0 +1,2 @@
+[CoreRedirects]
++ClassRedirects=(OldName="/Script/InputMappingContext.UnmapAction", NewName="/Script/InputMappingContext.UnmapAllKeysFromAction")
\ No newline at end of file
diff --git a/Engine/Plugins/EnhancedInput/Source/EnhancedInput/Private/InputMappingContext.cpp b/Engine/Plugins/EnhancedInput/Source/EnhancedInput/Private/InputMappingContext.cpp
--- a/Engine/Plugins/EnhancedInput/Source/EnhancedInput/Private/InputMappingContext.cpp
+++ b/Engine/Plugins/EnhancedInput/Source/EnhancedInput/Private/InputMappingContext.cpp
@@ -21,7 +21,7 @@ void UInputMappingContext::UnmapKey(const UInputAction* Action, FKey Key)
 	}
 }
 
-void UInputMappingContext::UnmapAction(const UInputAction* Action)
+void UInputMappingContext::UnmapAllKeysFromAction(const UInputAction* Action)
 {
 	int32 Found = Mappings.RemoveAllSwap([&Action](const FEnhancedActionKeyMapping& Mapping) { return Mapping.Action == Action; });
 	if (Found > 0)
diff --git a/Engine/Plugins/EnhancedInput/Source/EnhancedInput/Public/InputMappingContext.h b/Engine/Plugins/EnhancedInput/Source/EnhancedInput/Public/InputMappingContext.h
--- a/Engine/Plugins/EnhancedInput/Source/EnhancedInput/Public/InputMappingContext.h
+++ b/Engine/Plugins/EnhancedInput/Source/EnhancedInput/Public/InputMappingContext.h
@@ -59,11 +59,15 @@ public:
 	UFUNCTION(BlueprintCallable, Category = "Mapping")
 	void UnmapKey(const UInputAction* Action, FKey Key);
 
+	UE_DEPRECATED(5.1, "UnmapAction has been deprecated, please use UnmapAllKeysFromAction instead.")
+	UFUNCTION(BlueprintCallable, Category = "Mapping", meta=(DeprecatedFunction, DeprecatedMessage="UnmapAction has been deprecated, please use UnmapAllKeysFromAction instead."))
+	void UnmapAction(const UInputAction* Action) { UnmapAllKeysFromAction(Action); }
+	
 	/**
 	* Unmap all key maps to an action within the mapping context.
 	*/
 	UFUNCTION(BlueprintCallable, Category = "Mapping")
-	void UnmapAction(const UInputAction* Action);
+	void UnmapAllKeysFromAction(const UInputAction* Action);
 
 	/**
 	* Unmap everything within the mapping context.
