Enhanced Input: Add a "MetaData" pointer on the player mappable options. This will be really usefuil for devs if they wanted to define random stuff for each key mapping, like icons, maybe an ability asset, whatever you want.

#rnx

diff --git a/Engine/Plugins/EnhancedInput/Source/EnhancedInput/Public/EnhancedActionKeyMapping.h b/Engine/Plugins/EnhancedInput/Source/EnhancedInput/Public/EnhancedActionKeyMapping.h
--- a/Engine/Plugins/EnhancedInput/Source/EnhancedInput/Public/EnhancedActionKeyMapping.h
+++ b/Engine/Plugins/EnhancedInput/Source/EnhancedInput/Public/EnhancedActionKeyMapping.h
@@ -23,7 +23,7 @@ public:
 	
 	FPlayerMappableKeyOptions(const UInputAction* InAction = nullptr)
 	{
-		if(InAction)
+		if (InAction)
 		{
 			const FString& ActionName = InAction->GetName();
 			Name = FName(*ActionName);
@@ -36,11 +36,15 @@ public:
 		}
 	};
 
+	/** Metadata that can used to store any other related items to this key mapping such as icons, ability assets, etc. */
+	UPROPERTY(EditAnywhere, BlueprintReadOnly, Category = "Input|PlayerMappable")
+	TObjectPtr<UObject> Metadata = nullptr;
+	
 	/** A unique name for this player binding to be saved with. */
 	UPROPERTY(EditAnywhere, BlueprintReadOnly, Category = "Input|PlayerMappable")
 	FName Name;
 	
-	/** The display name that can  */
+	/** The localized display name of this key mapping. Use this when displaying the mappings to a user. */
 	UPROPERTY(EditAnywhere, BlueprintReadOnly, Category = "Input|PlayerMappable")
 	FText DisplayName;
 
