Add a style set for Enhanced Input so that the asset types have custom icons.

diff --git a/Engine/Plugins/EnhancedInput/Source/InputEditor/Private/InputEditorModule.cpp b/Engine/Plugins/EnhancedInput/Source/InputEditor/Private/InputEditorModule.cpp
--- a/Engine/Plugins/EnhancedInput/Source/InputEditor/Private/InputEditorModule.cpp
+++ b/Engine/Plugins/EnhancedInput/Source/InputEditor/Private/InputEditorModule.cpp
@@ -26,6 +26,10 @@
 #include "Widgets/Notifications/SNotificationList.h"
 #include "AssetTypeActions/AssetTypeActions_DataAsset.h"
 #include "EnhancedInputDeveloperSettings.h"
+#include "Styling/SlateStyle.h"
+#include "Styling/StyleColors.h"
+#include "Styling/SlateStyleMacros.h"
+#include "Styling/SlateStyleRegistry.h"
 
 #define LOCTEXT_NAMESPACE "InputEditor"
 
@@ -130,6 +134,34 @@ public:
 	virtual UClass* GetSupportedClass() const override { return UPlayerMappableInputConfig::StaticClass(); }
 };
 
+/** Custom style set for Enhanced Input */
+class FEnhancedInputSlateStyle final : public FSlateStyleSet
+{
+public:
+	FEnhancedInputSlateStyle()
+		: FSlateStyleSet("EnhancedInputEditor")
+	{
+		SetParentStyleName(FAppStyle::GetAppStyleSetName());
+
+		// The icons are located in /Engine/Plugins/EnhancedInput/Content/Editor/Slate/Icons
+		SetContentRoot(FPaths::EnginePluginsDir() / TEXT("EnhancedInput/Content/Editor/Slate"));
+		SetCoreContentRoot(FPaths::EngineContentDir() / TEXT("Slate"));
+
+		// Enhanced Input Editor icons
+		static const FVector2D Icon16 = FVector2D(16.0f, 16.0f);
+		static const FVector2D Icon64 = FVector2D(64.0f, 64.0f);
+
+		Set("ClassIcon.InputAction", new IMAGE_BRUSH_SVG("Icons/InputAction_16", Icon16));
+		Set("ClassThumbnail.InputAction", new IMAGE_BRUSH_SVG("Icons/InputAction_64", Icon64));
+		
+		Set("ClassIcon.InputMappingContext", new IMAGE_BRUSH_SVG("Icons/InputMappingContext_16", Icon16));
+		Set("ClassThumbnail.InputMappingContext", new IMAGE_BRUSH_SVG("Icons/InputMappingContext_64", Icon64));
+		
+		Set("ClassIcon.PlayerMappableInputConfig", new IMAGE_BRUSH_SVG("Icons/PlayerMappableInputConfig_16", Icon16));
+		Set("ClassThumbnail.PlayerMappableInputConfig", new IMAGE_BRUSH_SVG("Icons/PlayerMappableInputConfig_64", Icon64));			
+	}
+};
+
 void FInputEditorModule::StartupModule()
 {
 	// Register customizations
@@ -150,6 +182,10 @@ void FInputEditorModule::StartupModule()
 		//RegisterAssetTypeActions(AssetTools, MakeShareable(new FAssetTypeActions_InputTrigger));
 		//RegisterAssetTypeActions(AssetTools, MakeShareable(new FAssetTypeActions_InputModifier));
 	}
+
+	// Make a new style set for Enhanced Input, which will register any custom icons for the types in this plugin
+	StyleSet = MakeShared<FEnhancedInputSlateStyle>();
+	FSlateStyleRegistry::RegisterSlateStyle(*StyleSet.Get());
 }
 
 void FInputEditorModule::ShutdownModule()
@@ -176,6 +212,12 @@ void FInputEditorModule::ShutdownModule()
 	PropertyModule.UnregisterCustomPropertyTypeLayout("EnhancedActionKeyMapping");
 	PropertyModule.UnregisterCustomClassLayout("EnhancedInputDeveloperSettings");
 	PropertyModule.NotifyCustomizationModuleChanged();
+
+	// Unregister slate stylings
+	if (StyleSet.IsValid())
+	{
+		FSlateStyleRegistry::UnRegisterSlateStyle(*StyleSet.Get());
+	}
 }
 
 void FInputEditorModule::Tick(float DeltaTime)
diff --git a/Engine/Plugins/EnhancedInput/Source/InputEditor/Public/InputEditorModule.h b/Engine/Plugins/EnhancedInput/Source/InputEditor/Public/InputEditorModule.h
--- a/Engine/Plugins/EnhancedInput/Source/InputEditor/Public/InputEditorModule.h
+++ b/Engine/Plugins/EnhancedInput/Source/InputEditor/Public/InputEditorModule.h
@@ -41,6 +41,8 @@ private:
 	static EAssetTypeCategories::Type InputAssetsCategory;
 	
 	TArray<TSharedPtr<IAssetTypeActions>> CreatedAssetTypeActions;
+	
+	TSharedPtr<class FSlateStyleSet> StyleSet;
 };
 
 ////////////////////////////////////////////////////////////////////
