Add a "Scale by delta time" input modifier. This can be useful for if you wanted to have your look input or something in here, and then you don't need to add an additional " * deltatime " call in your input code every time.

diff --git a/Engine/Plugins/EnhancedInput/Source/EnhancedInput/Private/InputModifiers.cpp b/Engine/Plugins/EnhancedInput/Source/EnhancedInput/Private/InputModifiers.cpp
--- a/Engine/Plugins/EnhancedInput/Source/EnhancedInput/Private/InputModifiers.cpp
+++ b/Engine/Plugins/EnhancedInput/Source/EnhancedInput/Private/InputModifiers.cpp
@@ -20,6 +20,20 @@ FInputActionValue UInputModifierScalar::ModifyRaw_Implementation(const UEnhanced
 		return CurrentValue.Get<FVector>() * Scalar;
 	}
 	return CurrentValue;
+}
+
+/*
+* Scale by Delta Time
+*/
+
+FInputActionValue UInputModifierScaleByDeltaTime::ModifyRaw_Implementation(const UEnhancedPlayerInput* PlayerInput, FInputActionValue CurrentValue, float DeltaTime)
+{
+	// Don't try and scale bools
+	if (ensureMsgf(CurrentValue.GetValueType() != EInputActionValueType::Boolean, TEXT("Scale By Delta Time modifier doesn't support boolean values.")))
+	{
+		return CurrentValue.Get<FVector>() * DeltaTime;
+	}
+	return CurrentValue;
 };
 
 
diff --git a/Engine/Plugins/EnhancedInput/Source/EnhancedInput/Public/InputModifiers.h b/Engine/Plugins/EnhancedInput/Source/EnhancedInput/Public/InputModifiers.h
--- a/Engine/Plugins/EnhancedInput/Source/EnhancedInput/Public/InputModifiers.h
+++ b/Engine/Plugins/EnhancedInput/Source/EnhancedInput/Public/InputModifiers.h
@@ -132,6 +132,18 @@ protected:
 	virtual FInputActionValue ModifyRaw_Implementation(const UEnhancedPlayerInput* PlayerInput, FInputActionValue CurrentValue, float DeltaTime) override;
 };
 
+/** Scalar
+	*  Multiplies the input value by Delta Time this frame
+	*/
+UCLASS(NotBlueprintable, MinimalAPI, meta = (DisplayName = "Scale By Delta Time"))
+class UInputModifierScaleByDeltaTime : public UInputModifier
+{
+	GENERATED_BODY()
+	
+protected:
+	virtual FInputActionValue ModifyRaw_Implementation(const UEnhancedPlayerInput* PlayerInput, FInputActionValue CurrentValue, float DeltaTime) override;
+};
+
 /** Negate
 	*  Inverts input per axis
 	*/
