Add display names to the Enhanced Input developer settings and the EI editor settings

#rnx

diff --git a/Engine/Plugins/EnhancedInput/Source/EnhancedInput/Public/EnhancedInputDeveloperSettings.h b/Engine/Plugins/EnhancedInput/Source/EnhancedInput/Public/EnhancedInputDeveloperSettings.h
--- a/Engine/Plugins/EnhancedInput/Source/EnhancedInput/Public/EnhancedInputDeveloperSettings.h
+++ b/Engine/Plugins/EnhancedInput/Source/EnhancedInput/Public/EnhancedInputDeveloperSettings.h
@@ -1,25 +1,25 @@
-// Copyright Epic Games, Inc. All Rights Reserved.
-
-#pragma once
-
-#include "CoreMinimal.h"
-#include "Engine/DeveloperSettingsBackedByCVars.h"
-#include "Engine/PlatformSettings.h"
-#include "EnhancedInputDeveloperSettings.generated.h"
-
-/** Developer settings for Enhanced Input */
-UCLASS(config = Input, defaultconfig)
-class ENHANCEDINPUT_API UEnhancedInputDeveloperSettings : public UDeveloperSettingsBackedByCVars
-{
-	GENERATED_BODY()
-public:
-	
-	UEnhancedInputDeveloperSettings(const FObjectInitializer& Initializer);
-
-	/**
-	 * Platform specific settings for Enhanced Input.
-	 * @see UEnhancedInputPlatformSettings
-	 */
-	UPROPERTY(EditAnywhere, Category = "Input")
-	FPerPlatformSettings PlatformSettings;
+// Copyright Epic Games, Inc. All Rights Reserved.
+
+#pragma once
+
+#include "CoreMinimal.h"
+#include "Engine/DeveloperSettingsBackedByCVars.h"
+#include "Engine/PlatformSettings.h"
+#include "EnhancedInputDeveloperSettings.generated.h"
+
+/** Developer settings for Enhanced Input */
+UCLASS(config = Input, defaultconfig, meta = (DisplayName = "Enhanced Input"))
+class ENHANCEDINPUT_API UEnhancedInputDeveloperSettings : public UDeveloperSettingsBackedByCVars
+{
+	GENERATED_BODY()
+public:
+	
+	UEnhancedInputDeveloperSettings(const FObjectInitializer& Initializer);
+
+	/**
+	 * Platform specific settings for Enhanced Input.
+	 * @see UEnhancedInputPlatformSettings
+	 */
+	UPROPERTY(EditAnywhere, Category = "Input")
+	FPerPlatformSettings PlatformSettings;
 };
\ No newline at end of file
diff --git a/Engine/Plugins/EnhancedInput/Source/InputEditor/Public/EnhancedInputEditorSettings.h b/Engine/Plugins/EnhancedInput/Source/InputEditor/Public/EnhancedInputEditorSettings.h
--- a/Engine/Plugins/EnhancedInput/Source/InputEditor/Public/EnhancedInputEditorSettings.h
+++ b/Engine/Plugins/EnhancedInput/Source/InputEditor/Public/EnhancedInputEditorSettings.h
@@ -27,7 +27,7 @@
 };
 
 /** Settings for the Enhanced Input Editor Subsystem that are persistent between a project's users */
-UCLASS(config = Input, defaultconfig)
+UCLASS(config = Input, defaultconfig, meta = (DisplayName = "Enhanced Input (Editor Only)"))
 class UEnhancedInputEditorProjectSettings : public UDeveloperSettings
 {
 	GENERATED_BODY()
