Add an FText action description to UInputAction. Default the value to the name of the action with the common "IA_" prefix removed, and underscores replaced with white spaces. Change the category of the Input Mapping Context description field to be consistent with the input action's.

diff --git a/Engine/Plugins/EnhancedInput/Source/EnhancedInput/Private/InputAction.cpp b/Engine/Plugins/EnhancedInput/Source/EnhancedInput/Private/InputAction.cpp
--- a/Engine/Plugins/EnhancedInput/Source/EnhancedInput/Private/InputAction.cpp
+++ b/Engine/Plugins/EnhancedInput/Source/EnhancedInput/Private/InputAction.cpp
@@ -98,6 +98,18 @@ void FTriggerStateTracker::SetStateForNoTriggers(ETriggerState State)
 	NoTriggerState = State;
 }
 
+UInputAction::UInputAction()
+{
+	// Default to the name of this input action with the common prefix of "IA_" removed from it
+	FString DesiredName = GetNameSafe(this);
+	DesiredName.RemoveFromStart(FString(TEXT("IA_")));
+	
+	// Strip any underscores and replace them with a space as well
+	DesiredName = DesiredName.Replace(TEXT("_"), TEXT(" "));
+	
+	ActionDescription = FText::FromString(DesiredName);
+}
+
 #if WITH_EDITOR
 // Record input action property changes for later processing
 
diff --git a/Engine/Plugins/EnhancedInput/Source/EnhancedInput/Public/InputAction.h b/Engine/Plugins/EnhancedInput/Source/EnhancedInput/Public/InputAction.h
--- a/Engine/Plugins/EnhancedInput/Source/EnhancedInput/Public/InputAction.h
+++ b/Engine/Plugins/EnhancedInput/Source/EnhancedInput/Public/InputAction.h
@@ -20,6 +20,8 @@ class ENHANCEDINPUT_API UInputAction : public UDataAsset
 	GENERATED_BODY()
 
 public:
+	
+	UInputAction();
 
 #if WITH_EDITOR
 	// Track actions that have had their ValueType changed to update blueprints referencing them.
@@ -48,6 +50,10 @@ public:
 	UPROPERTY(EditAnywhere, BlueprintReadOnly, Category = Action)
 	EInputActionValueType ValueType = EInputActionValueType::Boolean;
 
+	// A localized descriptor of this input action
+	UPROPERTY(EditAnywhere, BlueprintReadOnly, Category = "Description")
+	FText ActionDescription = FText::GetEmpty();
+	
 	/**
 	* Trigger qualifiers. If any trigger qualifiers exist the action will not trigger unless:
 	* At least one Explicit trigger in this list is be met.
diff --git a/Engine/Plugins/EnhancedInput/Source/EnhancedInput/Public/InputMappingContext.h b/Engine/Plugins/EnhancedInput/Source/EnhancedInput/Public/InputMappingContext.h
--- a/Engine/Plugins/EnhancedInput/Source/EnhancedInput/Public/InputMappingContext.h
+++ b/Engine/Plugins/EnhancedInput/Source/EnhancedInput/Public/InputMappingContext.h
@@ -26,7 +26,7 @@ class ENHANCEDINPUT_API UInputMappingContext : public UDataAsset
 public:
 
 	// Localized context descriptor
-	UPROPERTY(BlueprintReadOnly, EditAnywhere, Category = "Config", DisplayName = "Description")
+	UPROPERTY(BlueprintReadOnly, EditAnywhere, Category = "Description", DisplayName = "Description")
 	FText ContextDescription;
 
 protected:
