For Input Action instance data that is flagged to trigger with "StartedAndTriggered" (i.e. started and triggered happen within the same frame) ensure that the "Started" delegate will always get called before the "Triggered" delegate. Before, "Triggered" was happening first because of the "TriggeredDelegates.Add(Binding->Clone())" logic, which would always put the "Started" event at the end of the array.

diff --git a/Engine/Plugins/EnhancedInput/Source/EnhancedInput/Private/EnhancedPlayerInput.cpp b/Engine/Plugins/EnhancedInput/Source/EnhancedInput/Private/EnhancedPlayerInput.cpp
--- a/Engine/Plugins/EnhancedInput/Source/EnhancedInput/Private/EnhancedPlayerInput.cpp
+++ b/Engine/Plugins/EnhancedInput/Source/EnhancedInput/Private/EnhancedPlayerInput.cpp
@@ -405,12 +405,21 @@ void UEnhancedPlayerInput::ProcessInputStack(const TArray<UInputComponent*>& Inp
 			// PERF: Lots of map lookups! Group EnhancedActionBindings by Action?
 			if (const FInputActionInstance* ActionData = FindActionInstanceData(Binding->GetAction()))
 			{
+				const ETriggerEvent BoundTriggerEvent = Binding->GetTriggerEvent();
 				// Raise appropriate delegate to report on event state
-				if (ActionData->TriggerEvent == Binding->GetTriggerEvent() ||
-					(Binding->GetTriggerEvent() == ETriggerEvent::Started && ActionData->TriggerEventInternal == ETriggerEventInternal::StartedAndTriggered))	// Triggering in a single tick should also fire the started event.
+				if (ActionData->TriggerEvent == BoundTriggerEvent ||
+					(BoundTriggerEvent == ETriggerEvent::Started && ActionData->TriggerEventInternal == ETriggerEventInternal::StartedAndTriggered))	// Triggering in a single tick should also fire the started event.
 				{
-					// Record intent to trigger
-					TriggeredDelegates.Add(Binding->Clone());
+					// Record intent to trigger started as well as triggered
+					// EmplaceAt 0 for the "Started" event it is always guaranteed to fire before Triggered
+					if (BoundTriggerEvent == ETriggerEvent::Started)
+					{
+						TriggeredDelegates.EmplaceAt(0, Binding->Clone());
+					}
+					else
+					{
+						TriggeredDelegates.Emplace(Binding->Clone());
+					}	
 				}
 			}
 		}
