Fix cook error that was happening because the ActionDescription is built from a FString, and thus did not have a null FText namespace key.

#rnx

diff --git a/Engine/Plugins/EnhancedInput/Source/EnhancedInput/Private/InputAction.cpp b/Engine/Plugins/EnhancedInput/Source/EnhancedInput/Private/InputAction.cpp
--- a/Engine/Plugins/EnhancedInput/Source/EnhancedInput/Private/InputAction.cpp
+++ b/Engine/Plugins/EnhancedInput/Source/EnhancedInput/Private/InputAction.cpp
@@ -98,18 +98,6 @@ void FTriggerStateTracker::SetStateForNoTriggers(ETriggerState State)
 	NoTriggerState = State;
 }
 
-UInputAction::UInputAction()
-{
-	// Default to the name of this input action with the common prefix of "IA_" removed from it
-	FString DesiredName = GetNameSafe(this);
-	DesiredName.RemoveFromStart(FString(TEXT("IA_")));
-	
-	// Strip any underscores and replace them with a space as well
-	DesiredName = DesiredName.Replace(TEXT("_"), TEXT(" "));
-	
-	ActionDescription = FText::FromString(DesiredName);
-}
-
 #if WITH_EDITOR
 // Record input action property changes for later processing
 
diff --git a/Engine/Plugins/EnhancedInput/Source/EnhancedInput/Public/InputAction.h b/Engine/Plugins/EnhancedInput/Source/EnhancedInput/Public/InputAction.h
--- a/Engine/Plugins/EnhancedInput/Source/EnhancedInput/Public/InputAction.h
+++ b/Engine/Plugins/EnhancedInput/Source/EnhancedInput/Public/InputAction.h
@@ -20,8 +20,6 @@ class ENHANCEDINPUT_API UInputAction : public UDataAsset
 	GENERATED_BODY()
 
 public:
-	
-	UInputAction();
 
 #if WITH_EDITOR
 	// Track actions that have had their ValueType changed to update blueprints referencing them.
