Exclude blueprint assets that have been marked for deletion in the content browser when we are building the enhanced input trigger defaults panel.

#rnx

diff --git a/Engine/Plugins/EnhancedInput/Source/InputEditor/Private/InputCustomizations.cpp b/Engine/Plugins/EnhancedInput/Source/InputEditor/Private/InputCustomizations.cpp
--- a/Engine/Plugins/EnhancedInput/Source/InputEditor/Private/InputCustomizations.cpp
+++ b/Engine/Plugins/EnhancedInput/Source/InputEditor/Private/InputCustomizations.cpp
@@ -87,6 +87,8 @@ void FEnhancedActionMappingCustomization::RemoveMappingButton_OnClick() const
 //////////////////////////////////////////////////////////
 // FEnhancedInputDeveloperSettingsCustomization
 
+TSet<FName> FEnhancedInputDeveloperSettingsCustomization::ExcludedAssetNames;
+
 FEnhancedInputDeveloperSettingsCustomization::~FEnhancedInputDeveloperSettingsCustomization()
 {
 	// Unregister settings panel listeners
@@ -111,6 +113,7 @@ void FEnhancedInputDeveloperSettingsCustomization::CustomizeDetails(IDetailLayou
 
 	TArray<UObject*> ModifierCDOs = GatherClassDetailsCDOs(UInputModifier::StaticClass());
 	TArray<UObject*> TriggerCDOs = GatherClassDetailsCDOs(UInputTrigger::StaticClass());
+	ExcludedAssetNames.Reset();
 	
 	// Add The modifier/trigger defaults that are generated via CDO to the details builder
 	CustomizeCDOValues(DetailBuilder, ModifierCategoryName, ModifierCDOs);
@@ -195,7 +198,7 @@ TArray<UObject*> FEnhancedInputDeveloperSettingsCustomization::GatherClassDetail
 	for (FAssetData& Asset : BlueprintAssetData)
 	{
 		FAssetDataTagMapSharedView::FFindTagResult Result = Asset.TagsAndValues.FindTag(TEXT("NativeParentClass"));
-		if (Result.IsSet())
+		if (Result.IsSet() && !ExcludedAssetNames.Contains(Asset.AssetName))
 		{
 			const FString ClassObjectPath = FPackageName::ExportTextPathToObjectPath(Result.GetValue());
 			const FString ClassName = FPackageName::ObjectPathToObjectName(ClassObjectPath);
@@ -237,7 +240,7 @@ TArray<UObject*> FEnhancedInputDeveloperSettingsCustomization::GatherClassDetail
 	return CDOs;
 }
 
-void FEnhancedInputDeveloperSettingsCustomization::RebuildDetailsViewForAsset(const FAssetData& AssetData)
+void FEnhancedInputDeveloperSettingsCustomization::RebuildDetailsViewForAsset(const FAssetData& AssetData, const bool bIsAssetBeingRemoved)
 {
 	// If the asset was a blueprint...
 	if (AssetData.AssetClass == UBlueprint::StaticClass()->GetFName())
@@ -253,6 +256,11 @@ void FEnhancedInputDeveloperSettingsCustomization::RebuildDetailsViewForAsset(co
 			{
 				if (ParentClass == UInputModifier::StaticClass() || ParentClass == UInputTrigger::StaticClass())
 				{
+					if (bIsAssetBeingRemoved)
+					{
+						ExcludedAssetNames.Add(AssetData.AssetName);
+					}
+					
 					if (IDetailLayoutBuilder* DetailBuilder = CachedDetailBuilder.Pin().Get())
 					{
 						DetailBuilder->ForceRefreshDetails();
diff --git a/Engine/Plugins/EnhancedInput/Source/InputEditor/Private/InputCustomizations.h b/Engine/Plugins/EnhancedInput/Source/InputEditor/Private/InputCustomizations.h
--- a/Engine/Plugins/EnhancedInput/Source/InputEditor/Private/InputCustomizations.h
+++ b/Engine/Plugins/EnhancedInput/Source/InputEditor/Private/InputCustomizations.h
@@ -70,12 +70,12 @@ private:
 	 * This will rebuild the Modifier and Trigger CDO views to make sure that
 	 * any newly added Blueprint
 	 */
-	void RebuildDetailsViewForAsset(const FAssetData& AssetData);
+	void RebuildDetailsViewForAsset(const FAssetData& AssetData, const bool bIsAssetBeingRemoved);
 
 	/** Callbacks that are triggered from the Asset Registry. */
-	void OnAssetAdded(const FAssetData& AssetData) { RebuildDetailsViewForAsset(AssetData); }
-	void OnAssetRemoved(const FAssetData& AssetData) { RebuildDetailsViewForAsset(AssetData); }
-	void OnAssetRenamed(const FAssetData& AssetData, const FString&) { RebuildDetailsViewForAsset(AssetData); }
+	void OnAssetAdded(const FAssetData& AssetData) { RebuildDetailsViewForAsset(AssetData, false); }
+	void OnAssetRemoved(const FAssetData& AssetData) { RebuildDetailsViewForAsset(AssetData, true); }
+	void OnAssetRenamed(const FAssetData& AssetData, const FString&) { RebuildDetailsViewForAsset(AssetData, false); }
 	
 	/**
 	 * Create a new category on the DetailBuilder and add each object in the given array as an external reference.
@@ -88,4 +88,11 @@ private:
 
 	// Cached details builder so that we can rebuild the details when a new BP asset is added
 	TWeakPtr<IDetailLayoutBuilder> CachedDetailBuilder;
+
+	/**
+	 * Populated by RebuildDetailsViewForAsset so that we can exclude any blueprint classes
+	 * that have been removed from the asset registry but are still loaded in memory.
+	 * This is reset when CustomizeCDOValues is called
+	 */
+	static TSet<FName> ExcludedAssetNames;
 };
\ No newline at end of file
