Export UInputTriggerChordAction and UInputTriggerTimedBase with the ENHANCEDINPUT_API.

diff --git a/Engine/Plugins/EnhancedInput/Source/EnhancedInput/Public/InputTriggers.h b/Engine/Plugins/EnhancedInput/Source/EnhancedInput/Public/InputTriggers.h
--- a/Engine/Plugins/EnhancedInput/Source/EnhancedInput/Public/InputTriggers.h
+++ b/Engine/Plugins/EnhancedInput/Source/EnhancedInput/Public/InputTriggers.h
@@ -165,8 +165,8 @@ Base class for building triggers that have firing conditions governed by elapsed
 This class transitions state to Ongoing once input is actuated, and will track Ongoing input time until input is released.
 Inheriting classes should provide the logic for Triggered transitions.
 */
-UCLASS(Abstract, MinimalAPI, Config = Input)	// TODO: Parent CLASS_Config flag is passed down to all UInputTrigger descendants, but ClassConfigName is only inherited by immediate children causing UHT to complain about a missing config file for any triggers based on this.
-class UInputTriggerTimedBase : public UInputTrigger
+UCLASS(Abstract, Config = Input)	// TODO: Parent CLASS_Config flag is passed down to all UInputTrigger descendants, but ClassConfigName is only inherited by immediate children causing UHT to complain about a missing config file for any triggers based on this.
+class ENHANCEDINPUT_API UInputTriggerTimedBase : public UInputTrigger
 {
 	GENERATED_BODY()
 
@@ -373,8 +373,8 @@ public:
  * UInputTriggerChordAction
  * Applies a chord action that must be triggering for this trigger's action to trigger
 */
-UCLASS(NotBlueprintable, MinimalAPI, meta = (DisplayName = "Chorded Action", NotInputConfigurable = "true"))
-class UInputTriggerChordAction : public UInputTrigger
+UCLASS(NotBlueprintable, meta = (DisplayName = "Chorded Action", NotInputConfigurable = "true"))
+class ENHANCEDINPUT_API UInputTriggerChordAction : public UInputTrigger
 {
 	GENERATED_BODY()
 
