Fix the display name on Player Mappable Configs in their asset creation factory. Now its consistent with the asset dipslay name.

#rnx

diff --git a/Engine/Plugins/EnhancedInput/Source/EnhancedInput/Private/PlayerMappableInputConfig.cpp b/Engine/Plugins/EnhancedInput/Source/EnhancedInput/Private/PlayerMappableInputConfig.cpp
--- a/Engine/Plugins/EnhancedInput/Source/EnhancedInput/Private/PlayerMappableInputConfig.cpp
+++ b/Engine/Plugins/EnhancedInput/Source/EnhancedInput/Private/PlayerMappableInputConfig.cpp
@@ -30,7 +30,7 @@ EDataValidationResult UPlayerMappableInputConfig::IsDataValid(TArray<FText>& Val
 		Result = CombineDataValidationResults(Result, EDataValidationResult::Invalid);
 	}
 	
-	// Check that every context is valid and all player bindable mappings have names
+	// Check that every context is valid and all player mappable mappings have names
 	for (const TPair<TObjectPtr<UInputMappingContext>, int32>& ContextPair : Contexts)
 	{
 		if(const UInputMappingContext* IMC = ContextPair.Key.Get())
@@ -44,7 +44,7 @@ EDataValidationResult UPlayerMappableInputConfig::IsDataValid(TArray<FText>& Val
 					Args.Add(TEXT("AssetPath"), FText::FromString(GetPathName()));
 					Args.Add(TEXT("ActionName"), FText::FromString(Mapping.Action->GetPathName()));
 
-					const FText NoMappingError = FText::Format(LOCTEXT("NoNameMappingError", "'{AssetPath}' has a player bindable key that has no name!"), Args);
+					const FText NoMappingError = FText::Format(LOCTEXT("NoNameMappingError", "'{AssetPath}' has a player mappable key that has no name!"), Args);
 					ValidationErrors.Emplace(NoMappingError);
 				}
 			}
diff --git a/Engine/Plugins/EnhancedInput/Source/EnhancedInput/Public/EnhancedActionKeyMapping.h b/Engine/Plugins/EnhancedInput/Source/EnhancedInput/Public/EnhancedActionKeyMapping.h
--- a/Engine/Plugins/EnhancedInput/Source/EnhancedInput/Public/EnhancedActionKeyMapping.h
+++ b/Engine/Plugins/EnhancedInput/Source/EnhancedInput/Public/EnhancedActionKeyMapping.h
@@ -77,7 +77,7 @@ struct FEnhancedActionKeyMapping
 
 	//ControllerId ControllerId;	// TODO: Controller id/player id (hybrid?) allowing binding multiple pads to a series of actions.
 
-	/** If true than this ActionKeyMapping will be exposed as a player bindable key */
+	/** If true than this ActionKeyMapping will be exposed as a player mappable key */
 	UPROPERTY(EditAnywhere, BlueprintReadOnly, Category = "Input|PlayerMappable")
 	uint8 bIsPlayerMappable : 1;
 
diff --git a/Engine/Plugins/EnhancedInput/Source/InputEditor/Private/InputEditorModule.cpp b/Engine/Plugins/EnhancedInput/Source/InputEditor/Private/InputEditorModule.cpp
--- a/Engine/Plugins/EnhancedInput/Source/InputEditor/Private/InputEditorModule.cpp
+++ b/Engine/Plugins/EnhancedInput/Source/InputEditor/Private/InputEditorModule.cpp
@@ -127,7 +127,7 @@ public:
 class FAssetTypeActions_PlayerMappableInputConfig : public FAssetTypeActions_DataAsset
 {
 public:
-	virtual FText GetName() const override { return NSLOCTEXT("AssetTypeActions", "AssetTypeActions_PlayerBindableInputConfig", "Player Bindable Input Config"); }
+	virtual FText GetName() const override { return NSLOCTEXT("AssetTypeActions", "AssetTypeActions_PlayerMappableInputConfig", "Player Mappable Input Config"); }
 	virtual uint32 GetCategories() override { return FInputEditorModule::GetInputAssetsCategory(); }
 	virtual FColor GetTypeColor() const override { return FColor(127, 255, 255); }
 	virtual FText GetAssetDescription(const FAssetData& AssetData) const override { return NSLOCTEXT("AssetTypeActions", "AssetTypeActions_PlayerBindableInputConfigDesc", "Represents one set of Player Mappable controller/keymappings"); }
