Fix LOCTEXT namespacing dup

#rnx

diff --git a/Engine/Plugins/EnhancedInput/Source/EnhancedInput/Private/PlayerMappableInputConfig.cpp b/Engine/Plugins/EnhancedInput/Source/EnhancedInput/Private/PlayerMappableInputConfig.cpp
--- a/Engine/Plugins/EnhancedInput/Source/EnhancedInput/Private/PlayerMappableInputConfig.cpp
+++ b/Engine/Plugins/EnhancedInput/Source/EnhancedInput/Private/PlayerMappableInputConfig.cpp
@@ -44,7 +44,7 @@ EDataValidationResult UPlayerMappableInputConfig::IsDataValid(TArray<FText>& Val
 					Args.Add(TEXT("AssetPath"), FText::FromString(GetPathName()));
 					Args.Add(TEXT("ActionName"), FText::FromString(Mapping.Action->GetPathName()));
 
-					const FText NoMappingError = FText::Format(LOCTEXT("NoNameMappingError", "'{AssetPath}' has a player mappable key that has no name!"), Args);
+					const FText NoMappingError = FText::Format(LOCTEXT("NoNamePlayerMappingError", "'{AssetPath}' has a player mappable key that has no name!"), Args);
 					ValidationErrors.Emplace(NoMappingError);
 				}
 			}
