Remove the warnings on unsupported input triggers because they are inconsistent and have a confusing UX. I am putting these behind a default-off CVar for now, until we can decide on a better route for determining when to warn.

#rnx

diff --git a/Engine/Plugins/EnhancedInput/Source/InputBlueprintNodes/Private/K2Node_EnhancedInputAction.cpp b/Engine/Plugins/EnhancedInput/Source/InputBlueprintNodes/Private/K2Node_EnhancedInputAction.cpp
--- a/Engine/Plugins/EnhancedInput/Source/InputBlueprintNodes/Private/K2Node_EnhancedInputAction.cpp
+++ b/Engine/Plugins/EnhancedInput/Source/InputBlueprintNodes/Private/K2Node_EnhancedInputAction.cpp
@@ -29,6 +29,16 @@ static const FName ElapsedSecondsPinName = TEXT("ElapsedSeconds");
 static const FName TriggeredSecondsPinName = TEXT("TriggeredSeconds");
 static const FName ActionValuePinName = TEXT("ActionValue");
 
+namespace UE::Input
+{
+	static bool bShouldWarnOnUnsupportedInputPin = false;
+	static FAutoConsoleVariableRef CVarShouldWarnOnUnsupportedInputPin(
+		TEXT("enhancedInput.bp.bShouldWarnOnUnsupportedInputPin"),
+		bShouldWarnOnUnsupportedInputPin,
+		TEXT("Should the Enhanced Input event node throw a warning if a \"Unsuported\" pin has a connection?"),
+		ECVF_Default);
+}
+
 UK2Node_EnhancedInputAction::UK2Node_EnhancedInputAction(const FObjectInitializer& ObjectInitializer)
 	: Super(ObjectInitializer)
 {
@@ -60,7 +70,7 @@ void UK2Node_EnhancedInputAction::AllocateDefaultPins()
 		NewPin->PinToolTip = EventEnum->GetToolTipTextByIndex((int32)Event).ToString();
 
 		// Add a special tooltip and display name for pins that are unsupported
-		if(!UInputTrigger::IsSupportedTriggerEvent(SupportedTriggerEvents, Event))
+		if (UE::Input::bShouldWarnOnUnsupportedInputPin && !UInputTrigger::IsSupportedTriggerEvent(SupportedTriggerEvents, Event))
 		{
 			static const FText UnsuportedTooltip = LOCTEXT("UnsupportedTooltip", "\n\nThis trigger event is not supported by the action! Add a supported trigger to enable this pin.");
 			NewPin->PinToolTip += UnsuportedTooltip.ToString();
@@ -281,7 +291,7 @@ void UK2Node_EnhancedInputAction::ExpandNode(FKismetCompilerContext& CompilerCon
 		{
 			ActivePins.Add(ActivePinData(InputActionPin, Event));
 			// Check if this exec pin is supported!
-			if(!UInputTrigger::IsSupportedTriggerEvent(SupportedTriggerEvents, Event))
+			if (UE::Input::bShouldWarnOnUnsupportedInputPin && !UInputTrigger::IsSupportedTriggerEvent(SupportedTriggerEvents, Event))
 			{
 				CompilerContext.MessageLog.Warning(*FText::Format(LOCTEXT("UnsuportedEventTypeOnAction", "'{0}'on @@ may not be executed because it is not a supported trigger on this action!"), InputActionPin->GetDisplayName()).ToString(), this);
 			}
