Add a missing super call to the PlayerMappableInputConfig constructor. Update comment on the deprecated property.

#rnx

diff --git a/Engine/Plugins/EnhancedInput/Source/EnhancedInput/Private/PlayerMappableInputConfig.cpp b/Engine/Plugins/EnhancedInput/Source/EnhancedInput/Private/PlayerMappableInputConfig.cpp
--- a/Engine/Plugins/EnhancedInput/Source/EnhancedInput/Private/PlayerMappableInputConfig.cpp
+++ b/Engine/Plugins/EnhancedInput/Source/EnhancedInput/Private/PlayerMappableInputConfig.cpp
@@ -6,6 +6,7 @@
 #define LOCTEXT_NAMESPACE "PlayerMappableInputConfig"
 
 UPlayerMappableInputConfig::UPlayerMappableInputConfig(const FObjectInitializer& ObjectInitializer)
+	: Super(ObjectInitializer)
 {
 	ResetToDefault();
 }
diff --git a/Engine/Plugins/EnhancedInput/Source/EnhancedInput/Public/PlayerMappableInputConfig.h b/Engine/Plugins/EnhancedInput/Source/EnhancedInput/Public/PlayerMappableInputConfig.h
--- a/Engine/Plugins/EnhancedInput/Source/EnhancedInput/Public/PlayerMappableInputConfig.h
+++ b/Engine/Plugins/EnhancedInput/Source/EnhancedInput/Public/PlayerMappableInputConfig.h
@@ -60,7 +60,7 @@ protected:
 	UPROPERTY(EditAnywhere, BlueprintReadOnly, Category = "Input|PlayerMappable")
 	FText ConfigDisplayName = FText::GetEmpty();
 
-	/**  */
+	/** A flag that can be used to mark this Input Config as deprecated to your player/designers. */
 	UPROPERTY(EditAnywhere, BlueprintReadOnly, Category = "Input|PlayerMappable")
 	bool bIsDeprecated = false;
 
