Do a pass on the alignment/padding of properties for Enhanced Input types.

UInputMappingContext -- No padding improvements, but the alignment will be better as the "Mappings" array will fill up to the 64 byte boundry

UInputAction -- Move the "Triggers" array up to the top to align at the 64 byte boundary, with the FText after it. In shipping builds, FText's are 24 bytes. This will actually result in more padding, but have it still be the same size. That means we have more room for any extra flags in the future!

FInputActionInstance -- Remove the two deprecated properties that are not in use anywhere. Seeing as we are coming out of experiemental, we should cut as much as we can now to save people from mucking stuff up in the production version. Went from 136 bytes to 96!!

FEnhancedActionKeyMapping -- Went from 144 bytes to 136 bytes. FKey is 80 bytes?! Ouch. Cut the amount of padding in half from 14 to 7

#rnx

diff --git a/Engine/Plugins/EnhancedInput/Source/EnhancedInput/Public/EnhancedActionKeyMapping.h b/Engine/Plugins/EnhancedInput/Source/EnhancedInput/Public/EnhancedActionKeyMapping.h
--- a/Engine/Plugins/EnhancedInput/Source/EnhancedInput/Public/EnhancedActionKeyMapping.h
+++ b/Engine/Plugins/EnhancedInput/Source/EnhancedInput/Public/EnhancedActionKeyMapping.h
@@ -67,6 +67,25 @@ struct FEnhancedActionKeyMapping
 {
 	GENERATED_BODY()
 
+	/** Options for making this a player mappable keymapping */
+	UPROPERTY(EditAnywhere, BlueprintReadOnly, Category = "Input|PlayerMappable", meta = (editCondition = "bIsPlayerMappable", DisplayAfter = "bIsPlayerMappable"))
+	FPlayerMappableKeyOptions PlayerMappableOptions;
+	
+	/**
+	* Trigger qualifiers. If any trigger qualifiers exist the mapping will not trigger unless:
+	* If there are any Explicit triggers in this list at least one of them must be met.
+	* All Implicit triggers in this list must be met.
+	*/
+	UPROPERTY(EditAnywhere, Instanced, BlueprintReadWrite, Category = "Input")
+	TArray<TObjectPtr<UInputTrigger>> Triggers;
+
+	/**
+	* Modifiers applied to the raw key value.
+	* These are applied sequentially in array order.
+	*/
+	UPROPERTY(EditAnywhere, Instanced, BlueprintReadWrite, Category = "Input")
+	TArray<TObjectPtr<UInputModifier>> Modifiers;
+	
 	/** Action to be affected by the key  */
 	UPROPERTY(EditAnywhere, BlueprintReadOnly, Category = "Input")
 	TObjectPtr<const UInputAction> Action = nullptr;
@@ -75,16 +94,10 @@ struct FEnhancedActionKeyMapping
 	UPROPERTY(EditAnywhere, BlueprintReadOnly, Category = "Input")
 	FKey Key;
 
-	//ControllerId ControllerId;	// TODO: Controller id/player id (hybrid?) allowing binding multiple pads to a series of actions.
-
 	/** If true than this ActionKeyMapping will be exposed as a player mappable key */
 	UPROPERTY(EditAnywhere, BlueprintReadOnly, Category = "Input|PlayerMappable")
 	uint8 bIsPlayerMappable : 1;
 
-	/** Options for making this a player mappable keymapping */
-	UPROPERTY(EditAnywhere, BlueprintReadOnly, Category = "Input|PlayerMappable", meta = (editCondition = "bIsPlayerMappable"))
-	FPlayerMappableKeyOptions PlayerMappableOptions;
-	
 	/**
 	 * If true, then this Key Mapping should be ignored. This is set to true if the key is down
 	 * during a rebuild of it's owning PlayerInput ControlMappings.
@@ -94,21 +107,6 @@ struct FEnhancedActionKeyMapping
 	UPROPERTY(Transient)
 	uint8 bShouldBeIgnored : 1;
 	
-	/**
-	* Trigger qualifiers. If any trigger qualifiers exist the mapping will not trigger unless:
-	* If there are any Explicit triggers in this list at least one of them must be met.
-	* All Implicit triggers in this list must be met.
-	*/
-	UPROPERTY(EditAnywhere, Instanced, BlueprintReadWrite, Category = "Input")
-	TArray<TObjectPtr<UInputTrigger>> Triggers;
-
-	/**
-	* Modifiers applied to the raw key value.
-	* These are applied sequentially in array order.
-	*/
-	UPROPERTY(EditAnywhere, Instanced, BlueprintReadWrite, Category = "Input")
-	TArray<TObjectPtr<UInputModifier>> Modifiers;
-
 	bool operator==(const FEnhancedActionKeyMapping& Other) const
 	{
 		return (Action == Other.Action &&
@@ -118,10 +116,10 @@ struct FEnhancedActionKeyMapping
 	}
 
 	FEnhancedActionKeyMapping(const UInputAction* InAction = nullptr, const FKey InKey = EKeys::Invalid)
-		: Action(InAction)
+		: PlayerMappableOptions(InAction)
+		, Action(InAction)
 		, Key(InKey)
 		, bIsPlayerMappable(false)
-		, PlayerMappableOptions(InAction)
 		, bShouldBeIgnored(false)
 	{}
 
diff --git a/Engine/Plugins/EnhancedInput/Source/EnhancedInput/Public/InputAction.h b/Engine/Plugins/EnhancedInput/Source/EnhancedInput/Public/InputAction.h
--- a/Engine/Plugins/EnhancedInput/Source/EnhancedInput/Public/InputAction.h
+++ b/Engine/Plugins/EnhancedInput/Source/EnhancedInput/Public/InputAction.h
@@ -32,6 +32,18 @@ public:
 	ETriggerEventsSupported GetSupportedTriggerEvents() const;
 #endif
 
+	/**
+	* Trigger qualifiers. If any trigger qualifiers exist the action will not trigger unless:
+	* At least one Explicit trigger in this list is be met.
+	* All Implicit triggers in this list are met.
+	*/
+	UPROPERTY(EditAnywhere, Instanced, BlueprintReadWrite, Category = Action)
+	TArray<TObjectPtr<UInputTrigger>> Triggers;
+	
+	// A localized descriptor of this input action
+	UPROPERTY(EditAnywhere, BlueprintReadOnly, Category = "Description")
+	FText ActionDescription = FText::GetEmpty();
+
 	// Should this action swallow any inputs bound to it or allow them to pass through to affect lower priority bound actions?
 	UPROPERTY(EditAnywhere, BlueprintReadOnly, Category = Action)
 	bool bConsumeInput = true;
@@ -47,25 +59,13 @@ public:
 	// The type that this action returns from a GetActionValue query or action event
 	UPROPERTY(EditAnywhere, BlueprintReadOnly, Category = Action)
 	EInputActionValueType ValueType = EInputActionValueType::Boolean;
-
-	// A localized descriptor of this input action
-	UPROPERTY(EditAnywhere, BlueprintReadOnly, Category = "Description")
-	FText ActionDescription = FText::GetEmpty();
 	
-	/**
-	* Trigger qualifiers. If any trigger qualifiers exist the action will not trigger unless:
-	* At least one Explicit trigger in this list is be met.
-	* All Implicit triggers in this list are met.
-	*/
-	UPROPERTY(EditAnywhere, Instanced, BlueprintReadWrite, Category = Action)
-	TArray<TObjectPtr<UInputTrigger>> Triggers;
-
 	/**
 	* Modifiers are applied to the final action value.
 	* These are applied sequentially in array order.
 	* They are applied on top of any FEnhancedActionKeyMapping modifiers that drove the initial input
 	*/
-	UPROPERTY(EditAnywhere, Instanced, BlueprintReadWrite, Category = Action)
+	UPROPERTY(EditAnywhere, Instanced, BlueprintReadWrite, Category = Action, meta=(DisplayAfter="Triggers"))
 	TArray<TObjectPtr<UInputModifier>> Modifiers;
 };
 
@@ -126,21 +126,23 @@ private:
 	// Internal trigger states
 	ETriggerState LastTriggerState = ETriggerState::None;
 	FTriggerStateTracker TriggerStateTracker;
-	ETriggerEventInternal TriggerEventInternal = ETriggerEventInternal(0);	// TODO: Expose access to ETriggerEventInternal?
+	ETriggerEventInternal TriggerEventInternal = ETriggerEventInternal(0);
 
 protected:
-	// TODO: Just hold a duplicate of the UInputAction in here?
-	// TODO: Restrict blueprint access to triggers and modifiers?
+
+	// Trigger state
+	UPROPERTY(BlueprintReadOnly, Transient, Category = Action)
+	ETriggerEvent TriggerEvent = ETriggerEvent::None;
+
+	// The last time that this evaluated to a Triggered State
+	UPROPERTY(BlueprintReadOnly, Transient, Category = Action)
+	float LastTriggeredWorldTime = 0.0f;
+
 	UPROPERTY(Instanced, BlueprintReadOnly, Category = Config)
 	TArray<TObjectPtr<UInputTrigger>> Triggers;
 
 	UPROPERTY(Instanced, BlueprintReadOnly, Category = Config)
 	TArray<TObjectPtr<UInputModifier>> Modifiers;
-	
-	UPROPERTY(meta = (DeprecatedProperty, DeprecationMessage = "Moved to Modifiers."))
-	TArray<TObjectPtr<UInputModifier>> PerInputModifiers_DEPRECATED;
-	UPROPERTY(meta = (DeprecatedProperty, DeprecationMessage = "Moved to Modifiers."))
-	TArray<TObjectPtr<UInputModifier>> FinalValueModifiers_DEPRECATED;
 
 	// Combined value of all inputs mapped to this action
 	FInputActionValue Value;
@@ -153,14 +155,6 @@ protected:
 	UPROPERTY(BlueprintReadOnly, Category = Action)
 	float ElapsedTriggeredTime = 0.f;
 
-	// The last time that this evaluated to a Triggered State
-	UPROPERTY(BlueprintReadOnly, Category = Action)
-	float LastTriggeredWorldTime = 0.0f;
-
-	// Trigger state
-	UPROPERTY(BlueprintReadOnly, Category = Action)
-	ETriggerEvent TriggerEvent = ETriggerEvent::None;
-
 public:
 	FInputActionInstance() = default;
 	FInputActionInstance(const UInputAction* InSourceAction);
@@ -183,9 +177,6 @@ public:
 	const TArray<UInputTrigger*>& GetTriggers() const { return Triggers; }
 	const TArray<UInputModifier*>& GetModifiers() const { return Modifiers; }
 
-	UE_DEPRECATED(4.26, "GetModifiers(EModifierExecutionPhase) is deprecated. Use GetModifiers()")
-	const TArray<UInputModifier*>& GetModifiers(EModifierExecutionPhase ForPhase) const { return Modifiers; }
-
 	// The source action that this instance is created from
 	const UInputAction* GetSourceAction() const { return SourceAction; }
 };
\ No newline at end of file
