Update the descriptoin on the Enhanced Input plugin. Update the docs URL to be correct

#rnx

diff --git a/Engine/Plugins/EnhancedInput/EnhancedInput.uplugin b/Engine/Plugins/EnhancedInput/EnhancedInput.uplugin
--- a/Engine/Plugins/EnhancedInput/EnhancedInput.uplugin
+++ b/Engine/Plugins/EnhancedInput/EnhancedInput.uplugin
@@ -3,11 +3,11 @@
 	"Version": 1,
 	"VersionName": "1.0",
 	"FriendlyName": "Enhanced Input",
-	"Description": "Prototype input handling.",
+	"Description": "Input handling that allows for contextual and dynamic mappings.",
 	"Category": "Input",
 	"CreatedBy": "Epic Games, Inc.",
 	"CreatedByURL": "http://epicgames.com",
-	"DocsURL": "",
+	"DocsURL": "https://docs.unrealengine.com/en-US/enhanced-input-in-unreal-engine/",
 	"MarketplaceURL": "",
 	"SupportURL": "",
 	"EnabledByDefault": true,
