Add SlateUserIndex to the platform devices debug command

#rnx

diff --git a/Engine/Plugins/EnhancedInput/Source/EnhancedInput/Private/EnhancedInputSubsystemInterfaceDebug.cpp b/Engine/Plugins/EnhancedInput/Source/EnhancedInput/Private/EnhancedInputSubsystemInterfaceDebug.cpp
--- a/Engine/Plugins/EnhancedInput/Source/EnhancedInput/Private/EnhancedInputSubsystemInterfaceDebug.cpp
+++ b/Engine/Plugins/EnhancedInput/Source/EnhancedInput/Private/EnhancedInputSubsystemInterfaceDebug.cpp
@@ -14,6 +14,7 @@
 #include "InputTriggers.h"
 #include "ImageUtils.h"
 #include "EnhancedInputPlatformSettings.h"
+#include "Framework/Application/SlateApplication.h"
 #include "GenericPlatform/GenericPlatformInputDeviceMapper.h"
 
 /* Shared input subsystem debug functionality.
@@ -310,9 +311,16 @@ void IEnhancedInputSubsystemInterface::ShowPlatformInputDebugInfo(class UCanvas*
 			static const float InputDeviceXOffset = 10.0f;
 			// Display the input device ID
 			FString DeviceIdString = FString::Printf(TEXT("Device ID: %d"), Device.GetId());
+			FString SlateUserIdString = TEXT("Slate UserID is unknown (Application is not initialized)");
+			if (FSlateApplication::IsInitialized())
+			{
+				const int32 SlateUserId = FSlateApplication::Get().GetUserIndexForInputDevice(Device);
+				SlateUserIdString = FString::Printf(TEXT("Slate UserId ID: %d"), SlateUserId);				
+			}
 			
 			DisplayDebugManager.SetDrawColor(FColor::White);
 			DisplayDebugManager.DrawString(DeviceIdString, InputDeviceXOffset);
+			DisplayDebugManager.DrawString(SlateUserIdString, InputDeviceXOffset);
 			
 			// Display the connection state
 			FString ConnectionStateString;
