Remove the deprecated Execution Phase from the UInputModifier.

#rnx

diff --git a/Engine/Plugins/EnhancedInput/Source/EnhancedInput/Public/InputModifiers.h b/Engine/Plugins/EnhancedInput/Source/EnhancedInput/Public/InputModifiers.h
--- a/Engine/Plugins/EnhancedInput/Source/EnhancedInput/Public/InputModifiers.h
+++ b/Engine/Plugins/EnhancedInput/Source/EnhancedInput/Public/InputModifiers.h
@@ -8,18 +8,6 @@
 
 class UEnhancedPlayerInput;
 
-// NOTE: Deprecated. Do not use.
-UENUM(BlueprintType)
-enum class EModifierExecutionPhase : uint8
-{
-	// Deprecated. Do not use.
-	PerInput,
-	// Deprecated. Do not use.
-	FinalValue,
-
-	NumPhases			UMETA(Hidden)
-};
-
 /**
 Base class for building modifiers.
 */
@@ -40,9 +28,6 @@ protected:
 		return CurrentValue;
 	}
 
-	UE_DEPRECATED(4.26, "Execution phase is deprecated.")
-	virtual EModifierExecutionPhase GetExecutionPhase_Implementation() const { return EModifierExecutionPhase::PerInput; }
-
 	virtual FLinearColor GetVisualizationColor_Implementation(FInputActionValue SampleValue, FInputActionValue FinalValue) const
 	{
 		const float Intensity = FMath::Min(1.f, FinalValue.GetMagnitude());	// TODO: 3D visualization!
@@ -59,13 +44,6 @@ public:
 	UFUNCTION(BlueprintCallable, BlueprintNativeEvent, Category = "Modifier")
 	FInputActionValue ModifyRaw(const UEnhancedPlayerInput* PlayerInput, FInputActionValue CurrentValue, float DeltaTime) const;
 
-	/**
-	 * GetExecutionPhase - Note: Deprecated
-	 */
-	UE_DEPRECATED(4.26, "Execution phase is deprecated. This call can be safely removed.")
-	UFUNCTION(BlueprintNativeEvent, Category = "Modifier", meta=(DeprecatedFunction, DeprecationMessage="Execution phase is deprecated."))
-	EModifierExecutionPhase GetExecutionPhase() const;
-
 	/**
 	 * Helper to allow debug visualization of the modifier.
 	 * @param SampleValue - The base input action value pre-modification (ranging -1 -> 1 across all applicable axes).
