Automatically upgrade projects with the default input class to use Enhanced Input when the module starts. By default this behavior is on, but it can be turned off with the cvar "EnhancedInput.bEnableAutoUpgrade".

We have to set the input settings properties on StartupModule like this because the base UInputSettings cannot reference the Enhanced Input plugin without creating a circular dependency.

Upgrading the project only when the current project is set to the default input classes is safe because it won't upgrade any project that has a custom class set on either property, and we can set on class but not the other with no issues.

diff --git a/Engine/Plugins/EnhancedInput/Source/InputEditor/Private/InputEditorModule.cpp b/Engine/Plugins/EnhancedInput/Source/InputEditor/Private/InputEditorModule.cpp
--- a/Engine/Plugins/EnhancedInput/Source/InputEditor/Private/InputEditorModule.cpp
+++ b/Engine/Plugins/EnhancedInput/Source/InputEditor/Private/InputEditorModule.cpp
@@ -30,11 +30,27 @@
 #include "Styling/StyleColors.h"
 #include "Styling/SlateStyleMacros.h"
 #include "Styling/SlateStyleRegistry.h"
+#include "GameFramework/InputSettings.h"
+#include "EnhancedInputComponent.h"
+#include "EnhancedPlayerInput.h"
+#include "Interfaces/IMainFrameModule.h"
 
 #define LOCTEXT_NAMESPACE "InputEditor"
 
+DEFINE_LOG_CATEGORY(LogEnhancedInputEditor);
+
 EAssetTypeCategories::Type FInputEditorModule::InputAssetsCategory;
 
+namespace UE::Input
+{
+	static bool bEnableAutoUpgradeToEnhancedInput = true;
+	static FAutoConsoleVariableRef CVarEnableAutoUpgradeToEnhancedInput(
+		TEXT("EnhancedInput.bEnableAutoUpgrade"),
+		bEnableAutoUpgradeToEnhancedInput,
+		TEXT("Should your project automatically be set to use Enhanced Input if it is currently using the legacy input system?"),
+		ECVF_Default);
+}
+
 IMPLEMENT_MODULE(FInputEditorModule, InputEditor)
 
 // Asset factories
@@ -186,6 +202,11 @@ void FInputEditorModule::StartupModule()
 	// Make a new style set for Enhanced Input, which will register any custom icons for the types in this plugin
 	StyleSet = MakeShared<FEnhancedInputSlateStyle>();
 	FSlateStyleRegistry::RegisterSlateStyle(*StyleSet.Get());
+
+	// Listen for when the editor is ready and then try to upgrade the input classes. We will send a slate notification
+	// if we upgrade the input classes, which is why we need to wait for the editor to be ready
+    IMainFrameModule& MainFrameModule = FModuleManager::LoadModuleChecked<IMainFrameModule>(TEXT("MainFrame"));
+    MainFrameModule.OnMainFrameCreationFinished().AddRaw(this, &FInputEditorModule::OnMainFrameCreationFinished);
 }
 
 void FInputEditorModule::ShutdownModule()
@@ -218,6 +239,77 @@ void FInputEditorModule::ShutdownModule()
 	{
 		FSlateStyleRegistry::UnRegisterSlateStyle(*StyleSet.Get());
 	}
+
+	// Remove any listeners to the editor startup delegate
+	if (FModuleManager::Get().IsModuleLoaded("MainFrame"))
+	{
+		IMainFrameModule::Get().OnMainFrameCreationFinished().RemoveAll(this);
+	}
+}
+
+void FInputEditorModule::OnMainFrameCreationFinished(TSharedPtr<SWindow> InRootWindow, bool bIsNewProjectWindow)
+{
+	AutoUpgradeDefaultInputClasses();
+}
+
+namespace  UE::Input
+{
+	static void ShowToast(const FText& ClassBeingChanged, const FString& NewClassName, const FString& OldClassName)
+	{
+		FFormatNamedArguments Args;
+		Args.Add(TEXT("ClassBeingChanged"), ClassBeingChanged);
+		Args.Add(TEXT("NewClass"), FText::FromString(NewClassName));
+		Args.Add(TEXT("OldClass"), FText::FromString(OldClassName));
+
+		const FText Message = FText::Format(LOCTEXT("EnhancedInputUpgradeToast", "Upgrading default {ClassBeingChanged} class from '{OldClass}' to '{NewClass}'"), Args);
+		
+		FNotificationInfo* Info = new FNotificationInfo(Message);
+		Info->ExpireDuration = 5.0f;
+		Info->bFireAndForget = true;
+		FSlateNotificationManager::Get().QueueNotification(Info);
+		
+		UE_LOG(LogEnhancedInputEditor, Log, TEXT("Upgrading Default %s class from '%s' to '%s'"), *ClassBeingChanged.ToString(), *OldClassName, *NewClassName);
+	}
+}
+
+void FInputEditorModule::AutoUpgradeDefaultInputClasses()
+{
+	if (!UE::Input::bEnableAutoUpgradeToEnhancedInput)
+	{
+		return;
+	}
+	
+	if (UInputSettings* InputSettings = GetMutableDefault<UInputSettings>())
+	{
+		const UClass* OriginalInputComponentClass = InputSettings->GetDefaultInputComponentClass();
+		const UClass* OriginalPlayerInputClass = InputSettings->GetDefaultPlayerInputClass();
+		bool bNeedsConfigSave = false;
+		
+		if (OriginalInputComponentClass && OriginalInputComponentClass == UInputComponent::StaticClass())
+		{
+			InputSettings->SetDefaultInputComponentClass(UEnhancedInputComponent::StaticClass());
+
+			static const FText ClassName = LOCTEXT("InputComponentClassLabel", "Input Component");
+			UE::Input::ShowToast(ClassName, UEnhancedInputComponent::StaticClass()->GetName(), OriginalInputComponentClass->GetName());
+			bNeedsConfigSave = true;
+		}
+
+		if (OriginalPlayerInputClass && OriginalPlayerInputClass == UPlayerInput::StaticClass())
+		{
+			InputSettings->SetDefaultPlayerInputClass(UEnhancedPlayerInput::StaticClass());
+
+			static const FText ClassName = LOCTEXT("PlayerInputClassLabel", "Player Input");
+			UE::Input::ShowToast(ClassName, UEnhancedPlayerInput::StaticClass()->GetName(), OriginalPlayerInputClass->GetName());
+			bNeedsConfigSave = true;
+		}
+
+		// Make sure that the config file gets updated with these new values
+		if (bNeedsConfigSave)
+		{
+			InputSettings->SaveConfig();
+			InputSettings->TryUpdateDefaultConfigFile();
+		}
+	}
 }
 
 void FInputEditorModule::Tick(float DeltaTime)
diff --git a/Engine/Plugins/EnhancedInput/Source/InputEditor/Public/InputEditorModule.h b/Engine/Plugins/EnhancedInput/Source/InputEditor/Public/InputEditorModule.h
--- a/Engine/Plugins/EnhancedInput/Source/InputEditor/Public/InputEditorModule.h
+++ b/Engine/Plugins/EnhancedInput/Source/InputEditor/Public/InputEditorModule.h
@@ -12,9 +12,13 @@
 
 #include "InputEditorModule.generated.h"
 
+DECLARE_LOG_CATEGORY_EXTERN(LogEnhancedInputEditor, Log, All);
+
 ////////////////////////////////////////////////////////////////////
 // FInputEditorModule
 
+class SWindow;
+
 class FInputEditorModule : public IModuleInterface, public FTickableEditorObject
 {
 public:
@@ -38,6 +42,11 @@ private:
 		CreatedAssetTypeActions.Add(Action);
 	}
 
+	void OnMainFrameCreationFinished(TSharedPtr<SWindow> InRootWindow, bool bIsNewProjectWindow);
+	
+	/** Automatically upgrade the current project to use Enhanced Input if it is currently set to the legacy input classes. */
+	void AutoUpgradeDefaultInputClasses();
+
 	static EAssetTypeCategories::Type InputAssetsCategory;
 	
 	TArray<TSharedPtr<IAssetTypeActions>> CreatedAssetTypeActions;
