Enhanced Input: making action mappings within a mapping context reorderable

diff --git a/Engine/Plugins/EnhancedInput/Source/EnhancedInput/Public/InputMappingContext.h b/Engine/Plugins/EnhancedInput/Source/EnhancedInput/Public/InputMappingContext.h
--- a/Engine/Plugins/EnhancedInput/Source/EnhancedInput/Public/InputMappingContext.h
+++ b/Engine/Plugins/EnhancedInput/Source/EnhancedInput/Public/InputMappingContext.h
@@ -35,6 +35,7 @@ public:
 	FText ContextDescription;
 
 	friend class FInputContextDetails;
+	friend class FActionMappingsNodeBuilderEx;
 
 public:
 
diff --git a/Engine/Plugins/EnhancedInput/Source/InputEditor/Private/ActionMappingDetails.cpp b/Engine/Plugins/EnhancedInput/Source/InputEditor/Private/ActionMappingDetails.cpp
--- a/Engine/Plugins/EnhancedInput/Source/InputEditor/Private/ActionMappingDetails.cpp
+++ b/Engine/Plugins/EnhancedInput/Source/InputEditor/Private/ActionMappingDetails.cpp
@@ -7,11 +7,13 @@
 #include "InputAction.h"
 #include "InputMappingContext.h"
 #include "IDetailChildrenBuilder.h"
+#include "IDetailDragDropHandler.h"
 #include "IDetailGroup.h"
 #include "IDetailPropertyRow.h"
 #include "IDocumentation.h"
 #include "PropertyCustomizationHelpers.h"
 #include "ScopedTransaction.h"
+#include "DragAndDrop/DecoratedDragDropOp.h"
 #include "Widgets/Input/SEditableTextBox.h"
 #include "Widgets/Layout/SBox.h"
 #include "Widgets/Text/STextBlock.h"
@@ -35,6 +37,235 @@ void FActionMappingsNodeBuilderEx::Tick(float DeltaTime)
 	HandleDelayedGroupExpansion();
 }
 
+int32 FActionMappingsNodeBuilderEx::GetNumGroupedMappings() const
+{
+	return GroupedMappings.Num();	
+}
+
+void FActionMappingsNodeBuilderEx::ReorderMappings(int32 OriginalIndex, int32 NewIndex)
+{
+	TArray<UObject*> OuterObjects;
+	ActionMappingsPropertyHandle->GetOuterObjects(OuterObjects);
+	if (OuterObjects.IsEmpty())
+	{
+		return;
+	}
+	UInputMappingContext* InputContext = CastChecked<UInputMappingContext>(OuterObjects[0]);
+	if (InputContext == nullptr)
+	{
+		return;
+	}
+
+	const FScopedTransaction Transaction(LOCTEXT("ReorderMapping_Transaction", "Move Mapping"));
+	InputContext->Modify();
+	ActionMappingsPropertyHandle->NotifyPreChange();
+
+	DelayedGroupExpansionStates.Emplace(nullptr, true);
+	
+	const UInputAction* ActionOfNewIndex = GroupedMappings[NewIndex].SharedAction;
+	const UInputAction* ActionOfOriginalIndex = GroupedMappings[OriginalIndex].SharedAction;
+	int32 NewMappingIndex = INDEX_NONE;
+	int32 OriginalMappingIndex = INDEX_NONE;
+
+	const TSharedPtr<IPropertyHandleArray> ActionMappingsArrayHandle = ActionMappingsPropertyHandle->AsArray();
+	uint32 NumMappings;
+	ActionMappingsArrayHandle->GetNumElements(NumMappings);
+	// loop through all mappings to find the first instance of the action we're moving to and from
+	for (uint32 Index = 0; Index < NumMappings; ++Index)
+	{
+		TSharedRef<IPropertyHandle> ActionMapping = ActionMappingsArrayHandle->GetElement(Index);
+		const UObject* Action;
+		const FPropertyAccess::Result Result = ActionMapping->GetChildHandle(GET_MEMBER_NAME_CHECKED(FEnhancedActionKeyMapping, Action))->GetValue(Action);
+
+		if (Result == FPropertyAccess::Success)
+		{
+			// Only care about the first time we find an action of that type
+			if (Action == ActionOfNewIndex && NewMappingIndex == INDEX_NONE)
+			{
+				NewMappingIndex = Index;
+			}
+			else if (Action == ActionOfOriginalIndex && OriginalMappingIndex == INDEX_NONE)
+			{
+				OriginalMappingIndex = Index;
+			}
+		}
+	}
+
+	// make sure we found indices for each action
+	if (NewMappingIndex == INDEX_NONE || OriginalIndex == INDEX_NONE)
+	{
+		return;
+	}
+
+	// figuring out if shifting items forward or backwards
+	if ((OriginalMappingIndex - NewMappingIndex) > 0)
+	{
+		FEnhancedActionKeyMapping OriginalKeyMapping = InputContext->Mappings[OriginalMappingIndex];
+		for (int32 Index = OriginalMappingIndex; Index > NewMappingIndex; --Index)
+		{
+			InputContext->Mappings[Index] = InputContext->Mappings[Index - 1];
+		}
+		InputContext->Mappings[NewMappingIndex] = OriginalKeyMapping;
+	}
+	else
+	{
+		FEnhancedActionKeyMapping OriginalKeyMapping = InputContext->Mappings[OriginalMappingIndex];
+		for (int32 Index = OriginalMappingIndex; Index < NewMappingIndex; ++Index)
+		{
+			InputContext->Mappings[Index] = InputContext->Mappings[Index + 1];
+		}
+		InputContext->Mappings[NewMappingIndex] = OriginalKeyMapping;
+	}
+
+	ActionMappingsPropertyHandle->NotifyPostChange(EPropertyChangeType::ArrayAdd);
+}
+
+/** Drag-and-drop operation that stores data about the source input mapping being dragged */
+class FInputMappingIndexDragDropOp : public FDecoratedDragDropOp
+{
+public:
+	DRAG_DROP_OPERATOR_TYPE(FUserDefinedEnumIndexDragDropOp, FDecoratedDragDropOp);
+
+	FInputMappingIndexDragDropOp(FMappingSet* InTargetMappingSet, int32 InInputMappingIndex, FActionMappingsNodeBuilderEx* InActionMappingNodeBuilder)
+		: TargetMappingSet(InTargetMappingSet)
+		, InputMappingIndex(InInputMappingIndex)
+		, ActionMappingNodeBuilder(InActionMappingNodeBuilder)
+	{
+		if (InTargetMappingSet->SharedAction)
+		{
+			InputMappingDisplayText = FText::FromName(InTargetMappingSet->SharedAction->GetFName());
+		}
+		else
+		{
+			InputMappingDisplayText = FText::FromString("None");
+		}
+		MouseCursor = EMouseCursor::GrabHandClosed;
+	}
+
+	void Init()
+	{
+		SetValidTarget(false);
+		SetupDefaults();
+		Construct();
+	}
+
+	void SetValidTarget(bool IsValidTarget)
+	{
+		FFormatNamedArguments Args;
+		Args.Add(TEXT("InputMappingName"), InputMappingDisplayText);
+
+		if (IsValidTarget)
+		{
+			CurrentHoverText = FText::Format(LOCTEXT("MoveInputMappingHere", "Move '{InputMappingName}' Here"), Args);
+			CurrentIconBrush = FAppStyle::GetBrush("Graph.ConnectorFeedback.OK");
+		}
+		else
+		{
+			CurrentHoverText = FText::Format(LOCTEXT("CannotMoveInputMappingHere", "Cannot Move '{InputMappingName}' Here"), Args);
+			CurrentIconBrush = FAppStyle::GetBrush("Graph.ConnectorFeedback.Error");
+		}
+	}
+
+	FMappingSet* GetMappingSet() const
+	{
+		return TargetMappingSet;
+	}
+
+	int32 GetInputMappingIndex() const
+	{
+		return InputMappingIndex;
+	}
+
+private:
+	FMappingSet* TargetMappingSet;
+	int32 InputMappingIndex;
+	FActionMappingsNodeBuilderEx* ActionMappingNodeBuilder;
+	FText InputMappingDisplayText;
+};
+
+/** Handler for customizing the drag-and-drop behavior for input mappings, allowing them to be reordered */
+class FInputMappingIndexDragDropHandler : public IDetailDragDropHandler
+{
+public:
+	FInputMappingIndexDragDropHandler(FMappingSet* InTargetMappingSet, int32 InTargetMappingIndex, FActionMappingsNodeBuilderEx* InActionMappingNodeBuilder)
+		: TargetInputMappingSet(InTargetMappingSet)
+		, InputMappingIndex(InTargetMappingIndex)
+		, ActionMappingNodeBuilder(InActionMappingNodeBuilder)
+	{
+		check(TargetInputMappingSet);
+		check(ActionMappingNodeBuilder);
+		check(InputMappingIndex >= 0 && InputMappingIndex < ActionMappingNodeBuilder->GetNumGroupedMappings());
+	}
+
+	virtual TSharedPtr<FDragDropOperation> CreateDragDropOperation() const override
+	{
+		TSharedPtr<FInputMappingIndexDragDropOp> DragOp = MakeShared<FInputMappingIndexDragDropOp>(TargetInputMappingSet, InputMappingIndex, ActionMappingNodeBuilder);
+		DragOp->Init();
+		return DragOp;
+	}
+
+	/** Compute new target index based on drop zone (above vs below) */
+	static int32 ComputeNewIndex(int32 OriginalIndex, int32 DropOntoIndex, EItemDropZone DropZone)
+	{
+		check(DropZone != EItemDropZone::OntoItem);
+
+		int32 NewIndex = DropOntoIndex;
+		if (DropZone == EItemDropZone::BelowItem)
+		{
+			// If the drop zone is below, then we actually move it to the next item's index
+			NewIndex++;
+		}
+		if (OriginalIndex < NewIndex)
+		{
+			// If the item is moved down the list, then all the other elements below it are shifted up one
+			NewIndex--;
+		}
+
+		return ensure(NewIndex >= 0) ? NewIndex : 0;
+	}
+
+	virtual bool AcceptDrop(const FDragDropEvent& DragDropEvent, EItemDropZone DropZone) const override
+	{
+		const TSharedPtr<FInputMappingIndexDragDropOp> DragOp = DragDropEvent.GetOperationAs<FInputMappingIndexDragDropOp>();
+		if (!DragOp.IsValid() || DropZone == EItemDropZone::OntoItem)
+		{
+			return false;
+		}
+
+		const int32 NewIndex = ComputeNewIndex(DragOp->GetInputMappingIndex(), InputMappingIndex, DropZone);
+		ActionMappingNodeBuilder->ReorderMappings(DragOp->GetInputMappingIndex(), NewIndex);
+		return true;
+	}
+
+	virtual TOptional<EItemDropZone> CanAcceptDrop(const FDragDropEvent& DragDropEvent, EItemDropZone DropZone) const override
+	{
+		const TSharedPtr<FInputMappingIndexDragDropOp> DragOp = DragDropEvent.GetOperationAs<FInputMappingIndexDragDropOp>();
+		if (!DragOp.IsValid())
+		{
+			return TOptional<EItemDropZone>();
+		}
+
+		// We're reordering, so there's no logical interpretation for dropping directly onto another item.
+		// Just change it to a drop-above in this case.
+		const EItemDropZone OverrideZone = (DropZone == EItemDropZone::BelowItem) ? EItemDropZone::BelowItem : EItemDropZone::AboveItem;
+		const int32 NewIndex = ComputeNewIndex(DragOp->GetInputMappingIndex(), InputMappingIndex, OverrideZone);
+
+		// Make sure that the new index is valid *and* that it represents an actual move from the current position.
+		if (NewIndex < 0 || NewIndex >= ActionMappingNodeBuilder->GetNumGroupedMappings() || NewIndex == DragOp->GetInputMappingIndex())
+		{
+			return TOptional<EItemDropZone>();
+		}
+
+		DragOp->SetValidTarget(true);
+		return OverrideZone;
+	}
+
+private:
+	FMappingSet* TargetInputMappingSet;
+	int32 InputMappingIndex;
+	FActionMappingsNodeBuilderEx* ActionMappingNodeBuilder;
+};
+
 void FActionMappingsNodeBuilderEx::GenerateHeaderRowContent(FDetailWidgetRow& NodeRow)
 {
 	TSharedRef<SWidget> AddButton = PropertyCustomizationHelpers::MakeAddButton(
@@ -133,6 +364,7 @@ void FActionMappingsNodeBuilderEx::GenerateChildContent(IDetailChildrenBuilder&
 		[
 			SNew(SHorizontalBox)
 			+ SHorizontalBox::Slot()
+			.Padding(InputConstants::PropertyPadding)
 			.AutoWidth()
 			[
 				SNew(SBox)
@@ -161,7 +393,8 @@ void FActionMappingsNodeBuilderEx::GenerateChildContent(IDetailChildrenBuilder&
 			[
 				RemoveButton
 			]
-		];
+		]
+		.DragDropHandler(MakeShared<FInputMappingIndexDragDropHandler>(&GroupedMappings[Index], Index, this));
 
 		for (int32 MappingIndex = 0; MappingIndex < MappingSet.Mappings.Num(); ++MappingIndex)
 		{
diff --git a/Engine/Plugins/EnhancedInput/Source/InputEditor/Private/ActionMappingDetails.h b/Engine/Plugins/EnhancedInput/Source/InputEditor/Private/ActionMappingDetails.h
--- a/Engine/Plugins/EnhancedInput/Source/InputEditor/Private/ActionMappingDetails.h
+++ b/Engine/Plugins/EnhancedInput/Source/InputEditor/Private/ActionMappingDetails.h
@@ -43,6 +43,8 @@ public:
 	virtual bool InitiallyCollapsed() const override { return true; };
 	virtual FName GetName() const override { return FName(TEXT("ActionMappings")); }
 
+	int32 GetNumGroupedMappings() const;
+	void ReorderMappings(int32 OriginalIndex, int32 NewIndex);
 private:
 	void AddActionMappingButton_OnClick();
 	void ClearActionMappingButton_OnClick();
