Add a "bShouldFireDelegatesInEditor" property to the Enhanced Input Component. This will make sure that only Input Components that have the flag specifically set when they are created will fire in-editor delegates with an Editor Script Guard. This ensures that no game-only input actions will be fired as editor delegates while keeping the implementation details hidden from the editor users.

diff --git a/Engine/Plugins/EnhancedInput/Source/EnhancedInput/Private/EnhancedInputComponent.cpp b/Engine/Plugins/EnhancedInput/Source/EnhancedInput/Private/EnhancedInputComponent.cpp
--- a/Engine/Plugins/EnhancedInput/Source/EnhancedInput/Private/EnhancedInputComponent.cpp
+++ b/Engine/Plugins/EnhancedInput/Source/EnhancedInput/Private/EnhancedInputComponent.cpp
@@ -15,6 +15,20 @@ UEnhancedInputComponent::UEnhancedInputComponent(const FObjectInitializer& Objec
 	bBlockInput = false;
 }
 
+void UEnhancedInputComponent::SetShouldFireDelegatesInEditor(const bool bInNewValue)
+{
+	bShouldFireDelegatesInEditor = bInNewValue;
+
+	// Propegate this info to any bound delegates
+	if (HasBindings())
+	{
+		for (auto& Binding : GetActionEventBindings())
+		{
+			Binding->SetShouldFireWithEditorScriptGuard(bShouldFireDelegatesInEditor);
+		}
+	}
+}
+
 bool UEnhancedInputComponent::HasBindings( ) const
 {
 	return GetActionEventBindings().Num() > 0 || GetActionValueBindings().Num() > 0 || GetDebugKeyBindings().Num() > 0 || Super::HasBindings();
diff --git a/Engine/Plugins/EnhancedInput/Source/EnhancedInput/Public/EnhancedInputComponent.h b/Engine/Plugins/EnhancedInput/Source/EnhancedInput/Public/EnhancedInputComponent.h
--- a/Engine/Plugins/EnhancedInput/Source/EnhancedInput/Public/EnhancedInputComponent.h
+++ b/Engine/Plugins/EnhancedInput/Source/EnhancedInput/Public/EnhancedInputComponent.h
@@ -26,6 +26,15 @@ DECLARE_DELEGATE_OneParam(FEnhancedInputActionHandlerValueSignature, const FInpu
 DECLARE_DELEGATE_OneParam(FEnhancedInputActionHandlerInstanceSignature, const FInputActionInstance&);	// Provides full access to value and timers
 DECLARE_DYNAMIC_DELEGATE_FourParams(FEnhancedInputActionHandlerDynamicSignature, FInputActionValue, ActionValue, float, ElapsedTime, float, TriggeredTime, const UInputAction*, SourceAction);
 
+// By default this behavior is on in the editor only. If you would like to turn this
+// off, then add the following to your Build.cs file:
+//		PublicDefinitions.Add("ENABLE_EDITOR_CALLABLE_INPUT_DELEGATES=0");
+#ifndef ENABLE_EDITOR_CALLABLE_INPUT_DELEGATES
+	#define ENABLE_EDITOR_CALLABLE_INPUT_DELEGATES	WITH_EDITOR
+#elif !WITH_EDITOR
+	#define ENABLE_EDITOR_CALLABLE_INPUT_DELEGATES	0
+#endif
+
 /** Unified storage for both native and dynamic delegates with any signature  */
 template<typename TSignature>
 struct TEnhancedInputUnifiedDelegate
@@ -34,7 +43,11 @@ private:
 	/** Holds the delegate to call. */
 	TSharedPtr<TSignature> Delegate;
 
+	/** Should this delegate fire with an Editor Script function guard? */
+	bool bShouldFireWithEditorScriptGuard = false;
+
 public:
+	
 	bool IsBound() const
 	{
 		return Delegate.IsValid() && Delegate->IsBound();
@@ -53,6 +66,10 @@ public:
 		}
 	}
 
+	void SetShouldFireWithEditorScriptGuard(const bool bNewValue) { bShouldFireWithEditorScriptGuard = bNewValue; }
+	
+	bool ShouldFireWithEditorScriptGuard() const { return bShouldFireWithEditorScriptGuard; }
+
 	/** Binds a native delegate, hidden for script delegates */
 	template<	typename UserClass,
 				typename TSig = TSignature>
@@ -80,12 +97,13 @@ public:
 		return *Delegate;
 	}
 
-	template<bool bUseScriptGuard = false, typename... TArgs>
+	template<typename... TArgs>
 	void Execute(TArgs... Args) const
 	{
 		if (IsBound())
 		{
-			if (bUseScriptGuard)
+#if ENABLE_EDITOR_CALLABLE_INPUT_DELEGATES
+			if (bShouldFireWithEditorScriptGuard)
 			{
 				FEditorScriptExecutionGuard ScriptGuard;
 				Delegate->Execute(Args...);
@@ -94,6 +112,10 @@ public:
 			{
 				Delegate->Execute(Args...);	
 			}
+#else
+			Delegate->Execute(Args...);
+#endif	// ENABLE_EDITOR_CALLABLE_INPUT_DELEGATES
+			
 		}
 	}
 };
@@ -147,6 +169,7 @@ public:
 
 	virtual void Execute(const FInputActionInstance& ActionData) const = 0;
 	virtual TUniquePtr<FEnhancedInputActionEventBinding> Clone() const = 0;
+	virtual void SetShouldFireWithEditorScriptGuard(const bool bNewValue) = 0;
 };
 
 /** Binds an action value for later reference. CurrentValue will be kept up to date with the value of the bound action */
@@ -214,9 +237,15 @@ template<typename TSignature>
 struct FEnhancedInputActionEventDelegateBinding : FEnhancedInputActionEventBinding
 {
 private:
-	FEnhancedInputActionEventDelegateBinding(const FEnhancedInputActionEventDelegateBinding<TSignature>& CloneFrom, EInputBindingClone Clone) : FEnhancedInputActionEventBinding(CloneFrom, Clone), Delegate(CloneFrom.Delegate) {}
+	FEnhancedInputActionEventDelegateBinding(const FEnhancedInputActionEventDelegateBinding<TSignature>& CloneFrom, EInputBindingClone Clone)
+		: FEnhancedInputActionEventBinding(CloneFrom, Clone), Delegate(CloneFrom.Delegate)
+	{
+		Delegate.SetShouldFireWithEditorScriptGuard(CloneFrom.Delegate.ShouldFireWithEditorScriptGuard());
+	}
 public:
-	FEnhancedInputActionEventDelegateBinding(const UInputAction* Action, ETriggerEvent InTriggerEvent) : FEnhancedInputActionEventBinding(Action, InTriggerEvent) {}
+	FEnhancedInputActionEventDelegateBinding(const UInputAction* Action, ETriggerEvent InTriggerEvent)
+		: FEnhancedInputActionEventBinding(Action, InTriggerEvent)
+	{}
 
 	// Implemented below.
 	virtual void Execute(const FInputActionInstance& ActionData) const override;
@@ -225,6 +254,11 @@ public:
 		return TUniquePtr<FEnhancedInputActionEventBinding>(new FEnhancedInputActionEventDelegateBinding<TSignature>(*this, EInputBindingClone::ForceClone));
 	}
 
+	virtual void SetShouldFireWithEditorScriptGuard(const bool bNewValue)
+	{
+		Delegate.SetShouldFireWithEditorScriptGuard(bNewValue);
+	}
+
 	TEnhancedInputUnifiedDelegate<TSignature> Delegate;
 };
 
@@ -251,37 +285,28 @@ public:
 
 // Action event delegate execution by signature.
 
-// By default this behavior is on in the editor only. If you would like to turn this
-// off, then add the following to your Build.cs file:
-//		PublicDefinitions.Add("ENABLE_EDITOR_CALLABLE_INPUT_DELEGATES=0");
-#ifndef ENABLE_EDITOR_CALLABLE_INPUT_DELEGATES
-	#define ENABLE_EDITOR_CALLABLE_INPUT_DELEGATES	WITH_EDITOR
-#elif !WITH_EDITOR
-	#define ENABLE_EDITOR_CALLABLE_INPUT_DELEGATES	0
-#endif
-
 template<>
 inline void FEnhancedInputActionEventDelegateBinding<FEnhancedInputActionHandlerSignature>::Execute(const FInputActionInstance& ActionData) const
 {
-	Delegate.Execute<ENABLE_EDITOR_CALLABLE_INPUT_DELEGATES>();
+	Delegate.Execute();
 }
 
 template<>
 inline void FEnhancedInputActionEventDelegateBinding<FEnhancedInputActionHandlerValueSignature>::Execute(const FInputActionInstance& ActionData) const
 {
-	Delegate.Execute<ENABLE_EDITOR_CALLABLE_INPUT_DELEGATES>(ActionData.GetValue());
+	Delegate.Execute(ActionData.GetValue());
 }
 
 template<>
 inline void FEnhancedInputActionEventDelegateBinding<FEnhancedInputActionHandlerInstanceSignature>::Execute(const FInputActionInstance& ActionData) const
 {
-	Delegate.Execute<ENABLE_EDITOR_CALLABLE_INPUT_DELEGATES>(ActionData);
+	Delegate.Execute(ActionData);
 }
 
 template<>
 inline void FEnhancedInputActionEventDelegateBinding<FEnhancedInputActionHandlerDynamicSignature>::Execute(const FInputActionInstance& ActionData) const
 {
-	Delegate.Execute<ENABLE_EDITOR_CALLABLE_INPUT_DELEGATES>(ActionData.GetValue(), ActionData.GetElapsedTime(), ActionData.GetTriggeredTime(), ActionData.GetSourceAction());
+	Delegate.Execute(ActionData.GetValue(), ActionData.GetElapsedTime(), ActionData.GetTriggeredTime(), ActionData.GetSourceAction());
 }
 
 
@@ -312,8 +337,19 @@ private:
 	/** Debug key bindings, available in non-shipping builds only. */
 	TArray<TUniquePtr<FInputDebugKeyBinding>> DebugKeyBindings;
 
+	/**
+	 * If true, then this input component's delegate bindings will fire with a FEditorScriptExecutionGuard around it.
+	 * This means that the events bound to this component can be triggered outside of "-game" scenarios, such as in
+	 * the level editor.
+	 * NOTE: Enabling this can interfere with events processed on gameplay actors (UE-148305) 
+	 */
+	bool bShouldFireDelegatesInEditor = false;
+
 public:
 
+	void SetShouldFireDelegatesInEditor(const bool bInNewValue);
+	
+	bool ShouldFireDelegatesInEditor() const { return bShouldFireDelegatesInEditor; }
 
 	/**
 	 * Checks whether this component has any input bindings.
@@ -377,6 +413,7 @@ public:
 	{																																											\
 		TUniquePtr<FEnhancedInputActionEventDelegateBinding<HANDLER_SIG>> AB = MakeUnique<FEnhancedInputActionEventDelegateBinding<HANDLER_SIG>>(Action, TriggerEvent);			\
 		AB->Delegate.BindDelegate<UserClass>(Object, Func);																														\
+		AB->Delegate.SetShouldFireWithEditorScriptGuard(bShouldFireDelegatesInEditor);																							\
 		return *EnhancedActionEventBindings.Add_GetRef(MoveTemp(AB));																											\
 	}
 
@@ -389,6 +426,7 @@ public:
 	{
 		TUniquePtr<FEnhancedInputActionEventDelegateBinding<FEnhancedInputActionHandlerSignature>> AB = MakeUnique<FEnhancedInputActionEventDelegateBinding<FEnhancedInputActionHandlerSignature>>(Action, TriggerEvent);
 		AB->Delegate.MakeDelegate().BindUObject(Object, Func, Vars...);
+		AB->Delegate.SetShouldFireWithEditorScriptGuard(bShouldFireDelegatesInEditor);
 		return *EnhancedActionEventBindings.Add_GetRef(MoveTemp(AB));
 	}
 
@@ -399,6 +437,7 @@ public:
 	{
 		TUniquePtr<FEnhancedInputActionEventDelegateBinding<FEnhancedInputActionHandlerDynamicSignature>> AB = MakeUnique<FEnhancedInputActionEventDelegateBinding<FEnhancedInputActionHandlerDynamicSignature>>(Action, TriggerEvent);
 		AB->Delegate.BindDelegate(Object, FunctionName);
+		AB->Delegate.SetShouldFireWithEditorScriptGuard(bShouldFireDelegatesInEditor);
 		return *EnhancedActionEventBindings.Add_GetRef(MoveTemp(AB));
 	}
 
@@ -426,6 +465,7 @@ public:
 #ifdef DEV_ONLY_KEY_BINDINGS_AVAILABLE
 		TUniquePtr<FInputDebugKeyDelegateBinding<FInputDebugKeyHandlerSignature>> KB = MakeUnique<FInputDebugKeyDelegateBinding<FInputDebugKeyHandlerSignature>>(Chord, KeyEvent, bExecuteWhenPaused);
 		KB->Delegate.BindDelegate(Object, Func);
+		KB->Delegate.SetShouldFireWithEditorScriptGuard(bShouldFireDelegatesInEditor);
 		return *DebugKeyBindings.Emplace_GetRef(MoveTemp(KB));
 #endif
 	}
@@ -438,6 +478,7 @@ public:
 #ifdef DEV_ONLY_KEY_BINDINGS_AVAILABLE
 		TUniquePtr<FInputDebugKeyDelegateBinding<FInputDebugKeyHandlerDynamicSignature>> KB = MakeUnique<FInputDebugKeyDelegateBinding<FInputDebugKeyHandlerDynamicSignature>>(Chord, KeyEvent, bExecuteWhenPaused);
 		KB->Delegate.BindDelegate(Object, FunctionName);
+		KB->Delegate.SetShouldFireWithEditorScriptGuard(bShouldFireDelegatesInEditor);
 		return *DebugKeyBindings.Emplace_GetRef(MoveTemp(KB));
 #endif
 	}
diff --git a/Engine/Plugins/EnhancedInput/Source/InputEditor/Private/EnhancedInputEditorSubsystem.cpp b/Engine/Plugins/EnhancedInput/Source/InputEditor/Private/EnhancedInputEditorSubsystem.cpp
--- a/Engine/Plugins/EnhancedInput/Source/InputEditor/Private/EnhancedInputEditorSubsystem.cpp
+++ b/Engine/Plugins/EnhancedInput/Source/InputEditor/Private/EnhancedInputEditorSubsystem.cpp
@@ -4,6 +4,7 @@
 #include "EnhancedPlayerInput.h"
 #include "InputMappingContext.h"
 #include "Components/InputComponent.h"
+#include "EnhancedInputComponent.h"
 #include "Subsystems/UnrealEditorSubsystem.h"
 #include "UnrealEdGlobals.h"
 #include "Editor/UnrealEdEngine.h"
@@ -142,6 +143,12 @@ void UEnhancedInputEditorSubsystem::PushInputComponent(UInputComponent* InInputC
 			UE_LOG(LogEditorInput, Log, TEXT("PUSHING AN INPUT COMPONENT! '%s'"), *InInputComponent->GetName());
 		}
 
+		// Mark this input component for firing delegates in the editor
+		if (UEnhancedInputComponent* EIC = Cast<UEnhancedInputComponent>(InInputComponent))
+		{
+			EIC->SetShouldFireDelegatesInEditor(true);
+		}
+
 		bool bPushed = false;
 		CurrentInputStack.RemoveSingle(InInputComponent);
 		for (int32 Index = CurrentInputStack.Num() - 1; Index >= 0; --Index)
