Adding a feature so you get visual feedback through a warning screen when binding a key that is already assigned in Lyra

diff --git a/Engine/Plugins/EnhancedInput/Source/EnhancedInput/Private/EnhancedInputSubsystemInterface.cpp b/Engine/Plugins/EnhancedInput/Source/EnhancedInput/Private/EnhancedInputSubsystemInterface.cpp
--- a/Engine/Plugins/EnhancedInput/Source/EnhancedInput/Private/EnhancedInputSubsystemInterface.cpp
+++ b/Engine/Plugins/EnhancedInput/Source/EnhancedInput/Private/EnhancedInputSubsystemInterface.cpp
@@ -417,6 +417,16 @@ int32 IEnhancedInputSubsystemInterface::RemovePlayerMappedKey(const FName Mappin
 	return NumMappingsApplied;
 }
 
+void IEnhancedInputSubsystemInterface::RemoveAllPlayerMappedKeys(const FModifyContextOptions& Options)
+{
+	if (UEnhancedPlayerInput* const PlayerInput = GetPlayerInput())
+	{
+		PlayerMappedSettings.Empty();
+	}
+
+	RequestRebuildControlMappings(Options);
+}
+
 void IEnhancedInputSubsystemInterface::AddPlayerMappableConfig(const UPlayerMappableInputConfig* Config, const FModifyContextOptions& Options)
 {
 	if(Config)
diff --git a/Engine/Plugins/EnhancedInput/Source/EnhancedInput/Public/EnhancedInputSubsystemInterface.h b/Engine/Plugins/EnhancedInput/Source/EnhancedInput/Public/EnhancedInputSubsystemInterface.h
--- a/Engine/Plugins/EnhancedInput/Source/EnhancedInput/Public/EnhancedInputSubsystemInterface.h
+++ b/Engine/Plugins/EnhancedInput/Source/EnhancedInput/Public/EnhancedInputSubsystemInterface.h
@@ -208,6 +208,14 @@ public:
 	UFUNCTION(BlueprintCallable, BlueprintCosmetic, Category = "Input|PlayerMappable", meta=(AutoCreateRefTerm = "Options"))
 	virtual int32 RemovePlayerMappedKey(const FName MappingName, const FModifyContextOptions& Options = FModifyContextOptions());
 	
+	/**
+	 * Remove All PlayerMappedKeys
+	 * Requests a rebuild of the player mappings. 
+	 *
+	 */
+	UFUNCTION(BlueprintCallable, BlueprintCosmetic, Category = "Input|PlayerMappable", meta=(AutoCreateRefTerm = "Options"))
+	virtual void RemoveAllPlayerMappedKeys(const FModifyContextOptions& Options = FModifyContextOptions());
+	
 	/** Adds all the input mapping contexts inside of this mappable config. */
 	UFUNCTION(BlueprintCallable, BlueprintCosmetic, Category = "Input|PlayerMappable", meta=(AutoCreateRefTerm = "Options"))
 	virtual void AddPlayerMappableConfig(const UPlayerMappableInputConfig* Config, const FModifyContextOptions& Options = FModifyContextOptions());
