Only attempt to save the input settings config file if it is writable. This solves an issue where a warning log would be printed if the file is currently checked out. This warning came from startup editor time upgrading of Input Classes to enhanced input, not just when you changed any input setting.

#rnx

diff --git a/Engine/Plugins/EnhancedInput/Source/InputEditor/InputEditor.Build.cs b/Engine/Plugins/EnhancedInput/Source/InputEditor/InputEditor.Build.cs
--- a/Engine/Plugins/EnhancedInput/Source/InputEditor/InputEditor.Build.cs
+++ b/Engine/Plugins/EnhancedInput/Source/InputEditor/InputEditor.Build.cs
@@ -37,6 +37,7 @@ namespace UnrealBuildTool.Rules
 					"AssetTools",
 					"DeveloperSettings",
 					"EditorSubsystem",
+					"SourceControl",
                 }
             );
         }
diff --git a/Engine/Plugins/EnhancedInput/Source/InputEditor/Private/InputEditorModule.cpp b/Engine/Plugins/EnhancedInput/Source/InputEditor/Private/InputEditorModule.cpp
--- a/Engine/Plugins/EnhancedInput/Source/InputEditor/Private/InputEditorModule.cpp
+++ b/Engine/Plugins/EnhancedInput/Source/InputEditor/Private/InputEditorModule.cpp
@@ -34,6 +34,8 @@
 #include "EnhancedInputComponent.h"
 #include "EnhancedPlayerInput.h"
 #include "Interfaces/IMainFrameModule.h"
+#include "SourceControlHelpers.h"
+#include "HAL/FileManager.h"
 
 #define LOCTEXT_NAMESPACE "InputEditor"
 
@@ -252,7 +254,7 @@ void FInputEditorModule::OnMainFrameCreationFinished(TSharedPtr<SWindow> InRootW
 	AutoUpgradeDefaultInputClasses();
 }
 
-namespace  UE::Input
+namespace UE::Input
 {
 	static void ShowToast(const FText& ClassBeingChanged, const FString& NewClassName, const FString& OldClassName)
 	{
@@ -306,8 +308,15 @@ void FInputEditorModule::AutoUpgradeDefaultInputClasses()
 		// Make sure that the config file gets updated with these new values
 		if (bNeedsConfigSave)
 		{
-			InputSettings->SaveConfig();
-			InputSettings->TryUpdateDefaultConfigFile();
+			const FString DefaultConfigFile = InputSettings->GetDefaultConfigFilename();
+			
+			// We can write to the file if it is not read only. If it is read only, then we can write to it if we successfully check it out with source control
+			bool bCanWriteToFile = !IFileManager::Get().IsReadOnly(*DefaultConfigFile) || USourceControlHelpers::CheckOutFile(DefaultConfigFile);
+			
+			if (bCanWriteToFile)
+			{
+				InputSettings->TryUpdateDefaultConfigFile(DefaultConfigFile, /* bWarnIfFail */ false);
+			}
 		}
 	}
 }
