Fix a build machine error where we would be trying to check out a file when source control is not enabled

diff --git a/Engine/Plugins/EnhancedInput/Source/InputEditor/Private/InputEditorModule.cpp b/Engine/Plugins/EnhancedInput/Source/InputEditor/Private/InputEditorModule.cpp
--- a/Engine/Plugins/EnhancedInput/Source/InputEditor/Private/InputEditorModule.cpp
+++ b/Engine/Plugins/EnhancedInput/Source/InputEditor/Private/InputEditorModule.cpp
@@ -311,7 +311,7 @@ void FInputEditorModule::AutoUpgradeDefaultInputClasses()
 			const FString DefaultConfigFile = InputSettings->GetDefaultConfigFilename();
 			
 			// We can write to the file if it is not read only. If it is read only, then we can write to it if we successfully check it out with source control
-			bool bCanWriteToFile = !IFileManager::Get().IsReadOnly(*DefaultConfigFile) || USourceControlHelpers::CheckOutFile(DefaultConfigFile);
+			bool bCanWriteToFile = !IFileManager::Get().IsReadOnly(*DefaultConfigFile) || (USourceControlHelpers::IsEnabled() && USourceControlHelpers::CheckOutFile(DefaultConfigFile));
 			
 			if (bCanWriteToFile)
 			{
