Add a ControlMappingsRebuiltDelegate delegate to the Enhanced Input local player subsystem. This allows you to listen for changes to control mappings or bindings for things like the UI and rebuild them.

You cannot put UPROPERTY's on interfaces, so we have to take the approach of having a virtual function call on the interface that any Subsystem's can implement for blueprint support.

diff --git a/Engine/Plugins/EnhancedInput/Source/EnhancedInput/Private/EnhancedInputModule.cpp b/Engine/Plugins/EnhancedInput/Source/EnhancedInput/Private/EnhancedInputModule.cpp
--- a/Engine/Plugins/EnhancedInput/Source/EnhancedInput/Private/EnhancedInputModule.cpp
+++ b/Engine/Plugins/EnhancedInput/Source/EnhancedInput/Private/EnhancedInputModule.cpp
@@ -294,6 +294,7 @@ void FEnhancedInputModule::Tick(float DeltaTime)
 		{
 			Subsystem->RebuildControlMappings();
 			Subsystem->TickForcedInput(DeltaTime);
+			Subsystem->HandleControlMappingRebuildDelegate();
 		});
 }
 
diff --git a/Engine/Plugins/EnhancedInput/Source/EnhancedInput/Private/EnhancedInputSubsystemInterface.cpp b/Engine/Plugins/EnhancedInput/Source/EnhancedInput/Private/EnhancedInputSubsystemInterface.cpp
--- a/Engine/Plugins/EnhancedInput/Source/EnhancedInput/Private/EnhancedInputSubsystemInterface.cpp
+++ b/Engine/Plugins/EnhancedInput/Source/EnhancedInput/Private/EnhancedInputSubsystemInterface.cpp
@@ -733,6 +733,7 @@ void IEnhancedInputSubsystemInterface::RebuildControlMappings()
 	
 	MappingRebuildPending = EInputMappingRebuildType::None;
 	bIgnoreAllPressedKeysUntilReleaseOnRebuild = true;
+	bControlMappingsRebuiltThisTick = true;
 }
 
 template<typename T>
@@ -784,6 +785,16 @@ void IEnhancedInputSubsystemInterface::TickForcedInput(float DeltaTime)
 	}
 }
 
+void IEnhancedInputSubsystemInterface::HandleControlMappingRebuildDelegate()
+{
+	if (bControlMappingsRebuiltThisTick)
+	{
+		ControlMappingsRebuiltThisFrame();
+		
+		bControlMappingsRebuiltThisTick = false;
+	}
+}
+
 void IEnhancedInputSubsystemInterface::ApplyForcedInput(const UInputAction* Action, FInputActionValue Value)
 {
 	check(Action);
diff --git a/Engine/Plugins/EnhancedInput/Source/EnhancedInput/Private/EnhancedInputSubsystems.cpp b/Engine/Plugins/EnhancedInput/Source/EnhancedInput/Private/EnhancedInputSubsystems.cpp
--- a/Engine/Plugins/EnhancedInput/Source/EnhancedInput/Private/EnhancedInputSubsystems.cpp
+++ b/Engine/Plugins/EnhancedInput/Source/EnhancedInput/Private/EnhancedInputSubsystems.cpp
@@ -27,4 +27,9 @@ UEnhancedPlayerInput* UEnhancedInputLocalPlayerSubsystem::GetPlayerInput() const
 		return Cast<UEnhancedPlayerInput>(PlayerController->PlayerInput);
 	}
 	return nullptr;
+}
+
+void UEnhancedInputLocalPlayerSubsystem::ControlMappingsRebuiltThisFrame()
+{
+	ControlMappingsRebuiltDelegate.Broadcast();
 }
\ No newline at end of file
diff --git a/Engine/Plugins/EnhancedInput/Source/EnhancedInput/Public/EnhancedInputSubsystemInterface.h b/Engine/Plugins/EnhancedInput/Source/EnhancedInput/Public/EnhancedInputSubsystemInterface.h
--- a/Engine/Plugins/EnhancedInput/Source/EnhancedInput/Public/EnhancedInputSubsystemInterface.h
+++ b/Engine/Plugins/EnhancedInput/Source/EnhancedInput/Public/EnhancedInputSubsystemInterface.h
@@ -231,6 +231,19 @@ protected:
 	 * in a subsytem implementing this interface, then you must call this function.
 	 */
 	void TickForcedInput(float DeltaTime);
+
+	/**
+	 * Called each tick from the Enhanced Input module.
+	 * 
+	 * If bControlMappingsRebuiltThisTick is true this tick, then call the ControlMappingsRebuiltThisFrame
+	 * function and reset the bControlMappingsRebuiltThisTick flag to false.
+	 * 
+	 * This gives an implementor of this interface an oppurtunity to add in a nice blueprint useable delegate.
+	 */
+	void HandleControlMappingRebuildDelegate();
+
+	/** Function that will be called when Control Mappings have been rebuilt this tick. */
+	virtual void ControlMappingsRebuiltThisFrame() {}
 	
 private:
 
@@ -275,6 +288,12 @@ private:
 
 	bool bMappingRebuildPending = false;
 
+	/**
+	 * If true then the control mappings have been rebuilt on this frame.
+	 * This is reset to false every tick in the EnhancedInputModule
+	 */
+	bool bControlMappingsRebuiltThisTick = false;
+
 	// Debug visualization implemented in EnhancedInputSubsystemsDebug.cpp
 	void ShowDebugInfo(class UCanvas* Canvas);
 	
diff --git a/Engine/Plugins/EnhancedInput/Source/EnhancedInput/Public/EnhancedInputSubsystems.h b/Engine/Plugins/EnhancedInput/Source/EnhancedInput/Public/EnhancedInputSubsystems.h
--- a/Engine/Plugins/EnhancedInput/Source/EnhancedInput/Public/EnhancedInputSubsystems.h
+++ b/Engine/Plugins/EnhancedInput/Source/EnhancedInput/Public/EnhancedInputSubsystems.h
@@ -19,5 +19,15 @@ public:
 
 	// Begin IEnhancedInputSubsystemInterface
 	virtual UEnhancedPlayerInput* GetPlayerInput() const override;
+	virtual void ControlMappingsRebuiltThisFrame() override;
 	// End IEnhancedInputSubsystemInterface
+
+	/** A delegate that will be called when control mappings have been rebuilt this frame. */
+	DECLARE_DYNAMIC_MULTICAST_DELEGATE(FOnControlMappingsRebuilt);
+
+	/**
+	 * Blueprint Event that is called at the end of any frame that Control Mappings have been rebuilt.
+	 */
+	UPROPERTY(BlueprintAssignable, DisplayName=OnControlMappingsRebuilt, Category = "Input")
+	FOnControlMappingsRebuilt ControlMappingsRebuiltDelegate;
 };
