Add accessors for the current Player Mapped Keys on Enhanced Input.

diff --git a/Engine/Plugins/EnhancedInput/Source/EnhancedInput/Private/EnhancedInputSubsystemInterface.cpp b/Engine/Plugins/EnhancedInput/Source/EnhancedInput/Private/EnhancedInputSubsystemInterface.cpp
--- a/Engine/Plugins/EnhancedInput/Source/EnhancedInput/Private/EnhancedInputSubsystemInterface.cpp
+++ b/Engine/Plugins/EnhancedInput/Source/EnhancedInput/Private/EnhancedInputSubsystemInterface.cpp
@@ -384,6 +384,24 @@ TArray<FKey> IEnhancedInputSubsystemInterface::QueryKeysMappedToAction(const UIn
 	return MappedKeys;
 }
 
+TArray<FEnhancedActionKeyMapping> IEnhancedInputSubsystemInterface::GetAllPlayerMappableActionKeyMappings() const
+{
+	TArray<FEnhancedActionKeyMapping> PlayerMappableMappings;
+	
+	if (const UEnhancedPlayerInput* const PlayerInput = GetPlayerInput())
+    {
+    	for (const FEnhancedActionKeyMapping& Mapping : PlayerInput->EnhancedActionMappings)
+		{
+			if (Mapping.bIsPlayerMappable)
+			{
+				PlayerMappableMappings.AddUnique(Mapping);
+			}
+		}
+    }
+	
+	return PlayerMappableMappings;
+}
+
 int32 IEnhancedInputSubsystemInterface::AddPlayerMappedKey(const FName MappingName, const FKey NewKey, const FModifyContextOptions& Options)
 {
 	int32 NumMappingsApplied = 0;
@@ -427,6 +445,16 @@ void IEnhancedInputSubsystemInterface::RemoveAllPlayerMappedKeys(const FModifyCo
 	RequestRebuildControlMappings(Options);
 }
 
+FKey IEnhancedInputSubsystemInterface::GetPlayerMappedKey(const FName MappingName) const
+{
+	if (const FKey* MappedKey = PlayerMappedSettings.Find(MappingName))
+	{
+		return FKey(*MappedKey);
+	}
+	
+	return EKeys::Invalid;
+}
+
 void IEnhancedInputSubsystemInterface::AddPlayerMappableConfig(const UPlayerMappableInputConfig* Config, const FModifyContextOptions& Options)
 {
 	if(Config)
diff --git a/Engine/Plugins/EnhancedInput/Source/EnhancedInput/Public/EnhancedInputSubsystemInterface.h b/Engine/Plugins/EnhancedInput/Source/EnhancedInput/Public/EnhancedInputSubsystemInterface.h
--- a/Engine/Plugins/EnhancedInput/Source/EnhancedInput/Public/EnhancedInputSubsystemInterface.h
+++ b/Engine/Plugins/EnhancedInput/Source/EnhancedInput/Public/EnhancedInputSubsystemInterface.h
@@ -190,6 +190,12 @@ public:
 	UFUNCTION(BlueprintCallable, BlueprintCosmetic, Category = "Input|Mapping Queries")
 	virtual TArray<FKey> QueryKeysMappedToAction(const UInputAction* Action) const;
 
+	/**
+	 * Get an array of the currently applied key mappings that are marked as Player Mappable.
+	 */
+	UFUNCTION(BlueprintCallable, BlueprintCosmetic, Category = "Input|Mapping Queries")
+	virtual TArray<FEnhancedActionKeyMapping> GetAllPlayerMappableActionKeyMappings() const;
+	
 	/**
 	 * Replace any currently applied mappings to this key mapping with the given new one.
 	 * Requests a rebuild of the player mappings. 
@@ -207,6 +213,15 @@ public:
 	 */
 	UFUNCTION(BlueprintCallable, BlueprintCosmetic, Category = "Input|PlayerMappable", meta=(AutoCreateRefTerm = "Options"))
 	virtual int32 RemovePlayerMappedKey(const FName MappingName, const FModifyContextOptions& Options = FModifyContextOptions());
+
+	/**
+	 * Get the player mapped key to the given mapping name. If there is not a player mapped key, then this will return
+	 * EKeys::Invalid.
+	 *
+	 * @param MappingName	The FName of the mapped key that would have been set with the AddPlayerMappedKey function.
+	 */
+	UFUNCTION(BlueprintCallable, BlueprintCosmetic, Category = "Input|PlayerMappable")
+	virtual FKey GetPlayerMappedKey(const FName MappingName) const;
 	
 	/**
 	 * Remove All PlayerMappedKeys
