Adding a bShouldAlwaysTick bool to allow InputTriggers to tick every frame

diff --git a/Engine/Plugins/EnhancedInput/Source/EnhancedInput/Private/EnhancedPlayerInput.cpp b/Engine/Plugins/EnhancedInput/Source/EnhancedInput/Private/EnhancedPlayerInput.cpp
--- a/Engine/Plugins/EnhancedInput/Source/EnhancedInput/Private/EnhancedPlayerInput.cpp
+++ b/Engine/Plugins/EnhancedInput/Source/EnhancedInput/Private/EnhancedPlayerInput.cpp
@@ -117,9 +117,22 @@ void UEnhancedPlayerInput::ProcessActionMappingEvent(const UInputAction* Action,
 	// Reset action data on the first event processed for the action this tick.
 	bool bResetActionData = !ActionsWithEventsThisTick.Contains(Action);
 	bool bMappingTriggersApplied = false;
+
+	bool bHasAnyAlwaysTickTriggers = false;
+	for (const UInputTrigger* Trigger : Triggers)
+	{
+		if (Trigger)
+		{
+			if (Trigger->bShouldAlwaysTick)
+			{
+				bHasAnyAlwaysTickTriggers = true;
+				break;
+			}
+		}
+	}
 	
 	// If the key state is changing or the key is actuated and being held (and not coming back up this tick) recalculate its value and resulting trigger state.
-	if (KeyEvent != EKeyEvent::None)
+	if (KeyEvent != EKeyEvent::None || bHasAnyAlwaysTickTriggers)
 	{
 		if (bResetActionData)
 		{
diff --git a/Engine/Plugins/EnhancedInput/Source/EnhancedInput/Public/InputTriggers.h b/Engine/Plugins/EnhancedInput/Source/EnhancedInput/Public/InputTriggers.h
--- a/Engine/Plugins/EnhancedInput/Source/EnhancedInput/Public/InputTriggers.h
+++ b/Engine/Plugins/EnhancedInput/Source/EnhancedInput/Public/InputTriggers.h
@@ -119,6 +119,12 @@ public:
 	UPROPERTY(EditAnywhere, Config, BlueprintReadWrite, Category = "Trigger Settings")
 	float ActuationThreshold = 0.5f;
 
+	/* Decides whether this trigger ticks every frame or not.
+	 * This WILL affect performance and should only be used in specific custom triggers.
+	 */
+	UPROPERTY(EditDefaultsOnly, Config, BlueprintReadOnly, Category = "Trigger Settings")
+	bool bShouldAlwaysTick = false;
+
 	/*
 	* Is the value passed in sufficiently large to be of interest to the trigger.
 	* This is a helper function that implements the most obvious (>=) interpretation of the actuation threshold.
