Fixup references to any renamed Input Actions in the Action Mapping Details for Enhanced Input. Previously, when you renamed an Input Action it would be set as "None" in any Input Mapping Context that referenced it. This problem only occured for Input Mapping Contexts that were currently open in the editor, and it just needed to be refreshed when one was renamed.

#rnx

diff --git a/Engine/Plugins/EnhancedInput/Source/InputEditor/Private/ActionMappingDetails.cpp b/Engine/Plugins/EnhancedInput/Source/InputEditor/Private/ActionMappingDetails.cpp
--- a/Engine/Plugins/EnhancedInput/Source/InputEditor/Private/ActionMappingDetails.cpp
+++ b/Engine/Plugins/EnhancedInput/Source/InputEditor/Private/ActionMappingDetails.cpp
@@ -17,6 +17,7 @@
 #include "Widgets/Input/SEditableTextBox.h"
 #include "Widgets/Layout/SBox.h"
 #include "Widgets/Text/STextBlock.h"
+#include "AssetRegistry/AssetRegistryModule.h"
 
 #define LOCTEXT_NAMESPACE "ActionMappingDetails"
 
@@ -26,6 +27,26 @@ FActionMappingsNodeBuilderEx::FActionMappingsNodeBuilderEx(IDetailLayoutBuilder*
 	: DetailLayoutBuilder(InDetailLayoutBuilder)
 	, ActionMappingsPropertyHandle(InPropertyHandle)
 {
+	// Support for updating references to renamed input actions
+	FAssetRegistryModule& AssetRegistryModule = FModuleManager::LoadModuleChecked<FAssetRegistryModule>(FName("AssetRegistry"));
+	IAssetRegistry& AssetRegistry = AssetRegistryModule.Get();
+	if (!AssetRegistry.OnAssetAdded().IsBoundToObject(this))
+	{
+		AssetRegistry.OnAssetRenamed().AddRaw(this, &FActionMappingsNodeBuilderEx::OnAssetRenamed);
+	}
+}
+
+FActionMappingsNodeBuilderEx::~FActionMappingsNodeBuilderEx()
+{
+	// Unregister settings panel listeners
+	if (FAssetRegistryModule* AssetRegistryModule = FModuleManager::GetModulePtr<FAssetRegistryModule>(FName("AssetRegistry")))
+	{
+		IAssetRegistry* AssetRegistry = AssetRegistryModule->TryGet();
+		if (AssetRegistry)
+		{
+			AssetRegistry->OnAssetRenamed().RemoveAll(this);
+		}
+	}
 }
 
 void FActionMappingsNodeBuilderEx::Tick(float DeltaTime)
@@ -340,6 +361,16 @@ FText FActionMappingsNodeBuilderEx::GetAddNewActionTooltip() const
 	}
 }
 
+void FActionMappingsNodeBuilderEx::OnAssetRenamed(const FAssetData& AssetData, const FString& OldName)
+{
+	// If this is an Input Action asset, then we need to check for any references to it and replace them
+	if (AssetData.GetClass() == UInputAction::StaticClass())
+	{
+		const FName OldPackageName = *FPackageName::ObjectPathToPackageName(OldName);
+		ActionsBeingRenamed.Add(AssetData.PackageName, OldPackageName);
+	}
+}
+
 void FActionMappingsNodeBuilderEx::GenerateChildContent(IDetailChildrenBuilder& ChildrenBuilder)
 {
 	RebuildGroupedMappings();
@@ -500,6 +531,12 @@ void FActionMappingsNodeBuilderEx::RemoveActionMappingGroupButton_OnClick(const
 
 bool FActionMappingsNodeBuilderEx::GroupsRequireRebuild() const
 {
+	// If any input actions have been renamed, then we need to update the details
+	if (!ActionsBeingRenamed.IsEmpty())
+	{
+		return true;
+	}
+	
 	for (int32 GroupIndex = 0; GroupIndex < GroupedMappings.Num(); ++GroupIndex)
 	{
 		const FMappingSet& MappingSet = GroupedMappings[GroupIndex];
@@ -519,6 +556,7 @@ bool FActionMappingsNodeBuilderEx::GroupsRequireRebuild() const
 void FActionMappingsNodeBuilderEx::RebuildGroupedMappings()
 {
 	GroupedMappings.Empty();
+	ActionsBeingRenamed.Empty();
 
 	TSharedPtr<IPropertyHandleArray> ActionMappingsArrayHandle = ActionMappingsPropertyHandle->AsArray();
 
diff --git a/Engine/Plugins/EnhancedInput/Source/InputEditor/Private/ActionMappingDetails.h b/Engine/Plugins/EnhancedInput/Source/InputEditor/Private/ActionMappingDetails.h
--- a/Engine/Plugins/EnhancedInput/Source/InputEditor/Private/ActionMappingDetails.h
+++ b/Engine/Plugins/EnhancedInput/Source/InputEditor/Private/ActionMappingDetails.h
@@ -33,6 +33,7 @@ class FActionMappingsNodeBuilderEx : public IDetailCustomNodeBuilder, public TSh
 {
 public:
 	FActionMappingsNodeBuilderEx(IDetailLayoutBuilder* InDetailLayoutBuilder, const TSharedPtr<IPropertyHandle>& InPropertyHandle);
+	virtual ~FActionMappingsNodeBuilderEx();
 
 	/** IDetailCustomNodeBuilder interface */
 	virtual void SetOnRebuildChildren(FSimpleDelegate InOnRebuildChildren) override { OnRebuildChildren = InOnRebuildChildren; }
@@ -67,6 +68,12 @@ private:
 	/** The tooltip attribute for the "Adds Action Mapping" button */
 	FText GetAddNewActionTooltip() const;
 
+	/**
+	 * Called by the asset registry when an asset is renamed. Replaces any references to the old Input Action with
+	 * the new one
+	 */
+	void OnAssetRenamed(const FAssetData& AssetData, const FString& OldName);
+
 private:
 	/** Called to rebuild the children of the detail tree */
 	FSimpleDelegate OnRebuildChildren;
@@ -80,4 +87,7 @@ private:
 	TArray<FMappingSet> GroupedMappings;
 
 	TArray<TPair<const UInputAction*, bool>> DelayedGroupExpansionStates;
+
+	/** A map of Input Action's that have been renamed. This map is a map of the New Name to the Old Name. */
+	TMap<FName, FName> ActionsBeingRenamed;
 };
\ No newline at end of file
