Remove EditDefaultsOnly from bShouldAlwaysTick on UInputTrigger. We only want people who make their own custom C++ or Blueprint trigger to be able to set this to true, and not have people toggling it in the editor.

#rnx

diff --git a/Engine/Plugins/EnhancedInput/Source/EnhancedInput/Public/InputTriggers.h b/Engine/Plugins/EnhancedInput/Source/EnhancedInput/Public/InputTriggers.h
--- a/Engine/Plugins/EnhancedInput/Source/EnhancedInput/Public/InputTriggers.h
+++ b/Engine/Plugins/EnhancedInput/Source/EnhancedInput/Public/InputTriggers.h
@@ -122,7 +122,7 @@ public:
 	/* Decides whether this trigger ticks every frame or not.
 	 * This WILL affect performance and should only be used in specific custom triggers.
 	 */
-	UPROPERTY(EditDefaultsOnly, Config, BlueprintReadOnly, Category = "Trigger Settings")
+	UPROPERTY(Config, BlueprintReadOnly, Category = "Trigger Settings")
 	bool bShouldAlwaysTick = false;
 
 	/*
