Extending Enhanced Input with a combo input trigger

diff --git a/Engine/Plugins/EnhancedInput/Source/EnhancedInput/Private/InputTriggers.cpp b/Engine/Plugins/EnhancedInput/Source/EnhancedInput/Private/InputTriggers.cpp
--- a/Engine/Plugins/EnhancedInput/Source/EnhancedInput/Private/InputTriggers.cpp
+++ b/Engine/Plugins/EnhancedInput/Source/EnhancedInput/Private/InputTriggers.cpp
@@ -207,4 +207,66 @@ ETriggerState UInputTriggerChordAction::UpdateState_Implementation(const UEnhanc
 	// Inherit state from the chorded action
 	const FInputActionInstance* EventData = PlayerInput->FindActionInstanceData(ChordAction);
 	return EventData ? EventData->TriggerStateTracker.GetState() : ETriggerState::None;
+}
+
+UInputTriggerComboAction::UInputTriggerComboAction()
+{
+	bShouldAlwaysTick = true;
+}
+
+ETriggerState UInputTriggerComboAction::UpdateState_Implementation(const UEnhancedPlayerInput* PlayerInput, FInputActionValue ModifiedValue, float DeltaTime)
+{
+	if (const UInputAction* CurrentAction = ComboActions[CurrentComboStepIndex].ComboStepAction)
+	{
+		for (const UInputAction* CancelAction : CancelActions)
+		{
+			if (CancelAction && CancelAction != CurrentAction)
+			{
+				const FInputActionInstance* CancelState = PlayerInput->FindActionInstanceData(CancelAction);
+				if (CancelState && CancelState->GetTriggerEvent() != ETriggerEvent::None)
+				{
+					// Cancel action firing!
+					CurrentComboStepIndex = 0;
+					CurrentAction = ComboActions[CurrentComboStepIndex].ComboStepAction;	// Reset for fallthrough
+				}
+			}
+		}
+
+		// Reset if we take too long to hit the next action
+		if (CurrentComboStepIndex > 0)
+		{
+			CurrentTimeBetweenComboSteps += DeltaTime;
+			if (CurrentTimeBetweenComboSteps >= ComboActions[CurrentComboStepIndex].TimeToPressKey)
+			{
+				CurrentComboStepIndex = 0;
+				CurrentAction = ComboActions[CurrentComboStepIndex].ComboStepAction;	// Reset for fallthrough			
+			}
+		}
+
+		const FInputActionInstance* CurrentState = PlayerInput->FindActionInstanceData(CurrentAction);
+		if (CurrentState && CurrentState->GetTriggerEvent() == ETriggerEvent::Completed) // + possibly Triggered
+		{
+			CurrentComboStepIndex++;
+			CurrentTimeBetweenComboSteps = 0;
+			if (CurrentComboStepIndex >= ComboActions.Num())
+			{
+				CurrentComboStepIndex = 0;
+				return ETriggerState::Triggered;
+			}
+		}
+
+		if (CurrentComboStepIndex > 0)
+		{
+			return ETriggerState::Ongoing;
+		}
+	
+		// Really should account for first combo action being mid-trigger...
+		const FInputActionInstance* InitialState = PlayerInput->FindActionInstanceData(ComboActions[0].ComboStepAction);
+		if (InitialState && InitialState->GetTriggerEvent() > ETriggerEvent::None) // || Cancelled!
+		{
+			return ETriggerState::Ongoing;
+		}
+		CurrentTimeBetweenComboSteps = 0;
+	}
+	return ETriggerState::None;
 };
\ No newline at end of file
diff --git a/Engine/Plugins/EnhancedInput/Source/EnhancedInput/Public/InputTriggers.h b/Engine/Plugins/EnhancedInput/Source/EnhancedInput/Public/InputTriggers.h
--- a/Engine/Plugins/EnhancedInput/Source/EnhancedInput/Public/InputTriggers.h
+++ b/Engine/Plugins/EnhancedInput/Source/EnhancedInput/Public/InputTriggers.h
@@ -8,6 +8,7 @@
 #include "InputTriggers.generated.h"
 
 class UEnhancedPlayerInput;
+class UInputAction;
 
 /**
 * Trigger states are a light weight interpretation of the provided input values used in trigger UpdateState responses.
@@ -393,7 +394,7 @@ public:
 
 	// The action that must be triggering for this trigger's action to trigger
 	UPROPERTY(EditInstanceOnly, BlueprintReadWrite, Category = "Trigger Settings", meta = (DisplayThumbnail = "false"))
-	TObjectPtr<const class UInputAction> ChordAction = nullptr;
+	TObjectPtr<const UInputAction> ChordAction = nullptr;
 };
 
 /** UInputTriggerChordBlocker
@@ -408,6 +409,56 @@ protected:
 	virtual ETriggerType GetTriggerType_Implementation() const override { return ETriggerType::Blocker; }
 };
 
+USTRUCT(BlueprintType)
+struct FInputComboStepData
+{
+	GENERATED_BODY()
+
+	// The action that must be triggering to progress the combo
+	UPROPERTY(EditInstanceOnly, BlueprintReadWrite, Category = "Trigger Settings", meta = (DisplayThumbnail = "false"))
+	TObjectPtr<const UInputAction> ComboStepAction = nullptr;
+
+	// Time to press the key and for it to count towards the combo
+	UPROPERTY(EditAnywhere, Config, BlueprintReadWrite, Category = "Trigger Settings")
+	float TimeToPressKey = 0.5f;
+};
+
+/**
+ * UInputTriggerComboAction
+ * All actions in the combo array must be pressed within a timeframe to trigger
+*/
+UCLASS(NotBlueprintable, meta = (DisplayName = "Combo", NotInputConfigurable = "true"))
+class ENHANCEDINPUT_API UInputTriggerComboAction : public UInputTrigger
+{
+	GENERATED_BODY()
+
+	UInputTriggerComboAction();
+
+protected:
+	// Implicit, so action cannot fire unless this is firing.
+	virtual ETriggerType GetTriggerType_Implementation() const override { return ETriggerType::Implicit; }
+
+	virtual ETriggerState UpdateState_Implementation(const UEnhancedPlayerInput* PlayerInput, FInputActionValue ModifiedValue, float DeltaTime) override;
+
+	// Keeps track of what action we're currently at in the combo
+	UPROPERTY(BlueprintReadOnly, Category = "Trigger Settings")
+	int32 CurrentComboStepIndex = 0;
+	
+	// Time elapsed between last combo InputAction trigger and current time
+	UPROPERTY(BlueprintReadOnly, Category = "Trigger Settings")
+	float CurrentTimeBetweenComboSteps = 0.0;
 
+public:
+	/**
+	 * List of input actions that need to be completed to trigger this action.
+	 * Input actions must be triggered in order (starting at index 0) to count towards the completion of the combo.
+	 */
+	UPROPERTY(EditAnywhere, BlueprintReadWrite, Category = "Trigger Settings", meta = (DisplayThumbnail = "false"))
+	TArray<FInputComboStepData> ComboActions;
+
+	// Actions that will cancel the combo if they are triggered
+    UPROPERTY(EditAnywhere, BlueprintReadWrite, Category = "Trigger Settings", meta = (DisplayThumbnail = "false"))
+    TArray<TObjectPtr<const UInputAction>> CancelActions;
+};
 
 //} // EnhancedInput
\ No newline at end of file
