Add a TitleProperty to the Combo steps so that it displays each step :)

#rnx

diff --git a/Engine/Plugins/EnhancedInput/Source/EnhancedInput/Public/InputTriggers.h b/Engine/Plugins/EnhancedInput/Source/EnhancedInput/Public/InputTriggers.h
--- a/Engine/Plugins/EnhancedInput/Source/EnhancedInput/Public/InputTriggers.h
+++ b/Engine/Plugins/EnhancedInput/Source/EnhancedInput/Public/InputTriggers.h
@@ -446,14 +446,14 @@ protected:
 	
 	// Time elapsed between last combo InputAction trigger and current time
 	UPROPERTY(BlueprintReadOnly, Category = "Trigger Settings")
-	float CurrentTimeBetweenComboSteps = 0.0;
+	float CurrentTimeBetweenComboSteps = 0.0f;
 
 public:
 	/**
 	 * List of input actions that need to be completed to trigger this action.
 	 * Input actions must be triggered in order (starting at index 0) to count towards the completion of the combo.
 	 */
-	UPROPERTY(EditAnywhere, BlueprintReadWrite, Category = "Trigger Settings", meta = (DisplayThumbnail = "false"))
+	UPROPERTY(EditAnywhere, BlueprintReadWrite, Category = "Trigger Settings", meta = (DisplayThumbnail = "false", TitleProperty = "ComboStepAction"))
 	TArray<FInputComboStepData> ComboActions;
 
 	// Actions that will cancel the combo if they are triggered
