Fixing a indexing problem in the combo array that crashes the editor

diff --git a/Engine/Plugins/EnhancedInput/Source/EnhancedInput/Private/InputTriggers.cpp b/Engine/Plugins/EnhancedInput/Source/EnhancedInput/Private/InputTriggers.cpp
--- a/Engine/Plugins/EnhancedInput/Source/EnhancedInput/Private/InputTriggers.cpp
+++ b/Engine/Plugins/EnhancedInput/Source/EnhancedInput/Private/InputTriggers.cpp
@@ -216,6 +216,12 @@ UInputTriggerComboAction::UInputTriggerComboAction()
 
 ETriggerState UInputTriggerComboAction::UpdateState_Implementation(const UEnhancedPlayerInput* PlayerInput, FInputActionValue ModifiedValue, float DeltaTime)
 {
+	if (ComboActions.IsEmpty())
+	{
+		ensureMsgf(false, TEXT("You must add combo actions to the UInputTriggerComboAction for it work properly! Exiting..."));
+		return ETriggerState::None;
+	}
+	
 	if (const UInputAction* CurrentAction = ComboActions[CurrentComboStepIndex].ComboStepAction)
 	{
 		for (const UInputAction* CancelAction : CancelActions)
