Update the comment on FindActionInstanceData on enhanced player input.

#rnx

diff --git a/Engine/Plugins/EnhancedInput/Source/EnhancedInput/Public/EnhancedPlayerInput.h b/Engine/Plugins/EnhancedInput/Source/EnhancedInput/Public/EnhancedPlayerInput.h
--- a/Engine/Plugins/EnhancedInput/Source/EnhancedInput/Public/EnhancedPlayerInput.h
+++ b/Engine/Plugins/EnhancedInput/Source/EnhancedInput/Public/EnhancedPlayerInput.h
@@ -30,7 +30,9 @@ class ENHANCEDINPUT_API UEnhancedPlayerInput : public UPlayerInput
 
 public:
 
-	// TODO: Can we avoid exposing this?
+	/**
+	* Returns the action instance data for the given input action if there is any. Returns nullptr if the action is not available.
+	*/
 	const FInputActionInstance* FindActionInstanceData(const UInputAction* ForAction) const { return ActionInstanceData.Find(ForAction); }
 
 	/** Retrieve the current value of an action for this player.
