Remove unnecessary comment in input triggers

#rnx

diff --git a/Engine/Plugins/EnhancedInput/Source/EnhancedInput/Public/InputTriggers.h b/Engine/Plugins/EnhancedInput/Source/EnhancedInput/Public/InputTriggers.h
--- a/Engine/Plugins/EnhancedInput/Source/EnhancedInput/Public/InputTriggers.h
+++ b/Engine/Plugins/EnhancedInput/Source/EnhancedInput/Public/InputTriggers.h
@@ -459,6 +459,4 @@ public:
 	// Actions that will cancel the combo if they are triggered
     UPROPERTY(EditAnywhere, BlueprintReadWrite, Category = "Trigger Settings", meta = (DisplayThumbnail = "false"))
     TArray<TObjectPtr<const UInputAction>> CancelActions;
-};
-
-//} // EnhancedInput
\ No newline at end of file
+};
\ No newline at end of file
