Update a file to use namespace for FMappingSet in a merge stream.

diff --git a/Engine/Plugins/EnhancedInput/Source/InputEditor/Private/ActionMappingDetails.cpp b/Engine/Plugins/EnhancedInput/Source/InputEditor/Private/ActionMappingDetails.cpp
--- a/Engine/Plugins/EnhancedInput/Source/InputEditor/Private/ActionMappingDetails.cpp
+++ b/Engine/Plugins/EnhancedInput/Source/InputEditor/Private/ActionMappingDetails.cpp
@@ -147,7 +147,7 @@ class FInputMappingIndexDragDropOp : public FDecoratedDragDropOp
 public:
 	DRAG_DROP_OPERATOR_TYPE(FUserDefinedEnumIndexDragDropOp, FDecoratedDragDropOp);
 
-	FInputMappingIndexDragDropOp(FMappingSet* InTargetMappingSet, int32 InInputMappingIndex, FActionMappingsNodeBuilderEx* InActionMappingNodeBuilder)
+	FInputMappingIndexDragDropOp(ActionMappingDetails::FMappingSet* InTargetMappingSet, int32 InInputMappingIndex, FActionMappingsNodeBuilderEx* InActionMappingNodeBuilder)
 		: TargetMappingSet(InTargetMappingSet)
 		, InputMappingIndex(InInputMappingIndex)
 		, ActionMappingNodeBuilder(InActionMappingNodeBuilder)
@@ -187,7 +187,7 @@ public:
 		}
 	}
 
-	FMappingSet* GetMappingSet() const
+	ActionMappingDetails::FMappingSet* GetMappingSet() const
 	{
 		return TargetMappingSet;
 	}
@@ -198,7 +198,7 @@ public:
 	}
 
 private:
-	FMappingSet* TargetMappingSet;
+	ActionMappingDetails::FMappingSet* TargetMappingSet;
 	int32 InputMappingIndex;
 	FActionMappingsNodeBuilderEx* ActionMappingNodeBuilder;
 	FText InputMappingDisplayText;
@@ -208,7 +208,7 @@ private:
 class FInputMappingIndexDragDropHandler : public IDetailDragDropHandler
 {
 public:
-	FInputMappingIndexDragDropHandler(FMappingSet* InTargetMappingSet, int32 InTargetMappingIndex, FActionMappingsNodeBuilderEx* InActionMappingNodeBuilder)
+	FInputMappingIndexDragDropHandler(ActionMappingDetails::FMappingSet* InTargetMappingSet, int32 InTargetMappingIndex, FActionMappingsNodeBuilderEx* InActionMappingNodeBuilder)
 		: TargetInputMappingSet(InTargetMappingSet)
 		, InputMappingIndex(InTargetMappingIndex)
 		, ActionMappingNodeBuilder(InActionMappingNodeBuilder)
@@ -282,7 +282,7 @@ public:
 	}
 
 private:
-	FMappingSet* TargetInputMappingSet;
+	ActionMappingDetails::FMappingSet* TargetInputMappingSet;
 	int32 InputMappingIndex;
 	FActionMappingsNodeBuilderEx* ActionMappingNodeBuilder;
 };
