Add required namespace to InputConstants::PropertyPadding

diff --git a/Engine/Plugins/EnhancedInput/Source/InputEditor/Private/ActionMappingDetails.cpp b/Engine/Plugins/EnhancedInput/Source/InputEditor/Private/ActionMappingDetails.cpp
--- a/Engine/Plugins/EnhancedInput/Source/InputEditor/Private/ActionMappingDetails.cpp
+++ b/Engine/Plugins/EnhancedInput/Source/InputEditor/Private/ActionMappingDetails.cpp
@@ -395,7 +395,7 @@ void FActionMappingsNodeBuilderEx::GenerateChildContent(IDetailChildrenBuilder&
 		[
 			SNew(SHorizontalBox)
 			+ SHorizontalBox::Slot()
-			.Padding(InputConstants::PropertyPadding)
+			.Padding(ActionMappingDetails::InputConstants::PropertyPadding)
 			.AutoWidth()
 			[
 				SNew(SBox)
