Delete deprecated 5.0 functions from enhanced input. Enhanced Input was experimental when these functions were deprecated, and now that it is production ready this will give us a clean slate to work on. These functions were removed:

- Old version of AddMappingContext
- Old version of RemoveMappingContext
- Old version of RequestRebuildControlMappings

diff --git a/Engine/Plugins/EnhancedInput/Source/EnhancedInput/Private/EnhancedInputSubsystemInterface.cpp b/Engine/Plugins/EnhancedInput/Source/EnhancedInput/Private/EnhancedInputSubsystemInterface.cpp
--- a/Engine/Plugins/EnhancedInput/Source/EnhancedInput/Private/EnhancedInputSubsystemInterface.cpp
+++ b/Engine/Plugins/EnhancedInput/Source/EnhancedInput/Private/EnhancedInputSubsystemInterface.cpp
@@ -52,14 +52,6 @@ void IEnhancedInputSubsystemInterface::ClearAllMappings()
 	}
 }
 
-void IEnhancedInputSubsystemInterface::AddMappingContext(const UInputMappingContext* MappingContext, int32 Priority, const bool bIgnoreAllPressedKeysUntilRelease)
-{
-	FModifyContextOptions Options {};
-	Options.bIgnoreAllPressedKeysUntilRelease = bIgnoreAllPressedKeysUntilRelease;
-	
-	AddMappingContext(MappingContext, Priority, Options);
-}
-
 void IEnhancedInputSubsystemInterface::AddMappingContext(const UInputMappingContext* MappingContext, int32 Priority, const FModifyContextOptions& Options)
 {
 	// Layer mappings on top of existing mappings
@@ -73,14 +65,6 @@ void IEnhancedInputSubsystemInterface::AddMappingContext(const UInputMappingCont
 	}
 }
 
-void IEnhancedInputSubsystemInterface::RemoveMappingContext(const UInputMappingContext* MappingContext, const bool bIgnoreAllPressedKeysUntilRelease)
-{
-	FModifyContextOptions Options {};
-	Options.bIgnoreAllPressedKeysUntilRelease = bIgnoreAllPressedKeysUntilRelease;
-	
-	RemoveMappingContext(MappingContext, Options);
-}
-
 void IEnhancedInputSubsystemInterface::RemoveMappingContext(const UInputMappingContext* MappingContext, const FModifyContextOptions& Options)
 {
 	if (MappingContext)
@@ -93,15 +77,6 @@ void IEnhancedInputSubsystemInterface::RemoveMappingContext(const UInputMappingC
 	}
 }
 
-void IEnhancedInputSubsystemInterface::RequestRebuildControlMappings(bool bForceImmediately, const bool bIgnoreAllPressedKeysUntilRelease)
-{
-	FModifyContextOptions Options;
-	Options.bForceImmediately = bForceImmediately;
-	Options.bIgnoreAllPressedKeysUntilRelease = bIgnoreAllPressedKeysUntilRelease;
-	
-	RequestRebuildControlMappings(Options);
-}
-
 void IEnhancedInputSubsystemInterface::RequestRebuildControlMappings(const FModifyContextOptions& Options, EInputMappingRebuildType MappingRebuildType)
 {
 	bMappingRebuildPending = true;
diff --git a/Engine/Plugins/EnhancedInput/Source/EnhancedInput/Public/EnhancedInputSubsystemInterface.h b/Engine/Plugins/EnhancedInput/Source/EnhancedInput/Public/EnhancedInputSubsystemInterface.h
--- a/Engine/Plugins/EnhancedInput/Source/EnhancedInput/Public/EnhancedInputSubsystemInterface.h
+++ b/Engine/Plugins/EnhancedInput/Source/EnhancedInput/Public/EnhancedInputSubsystemInterface.h
@@ -96,16 +96,6 @@ public:
 	 */
 	UFUNCTION(BlueprintCallable, BlueprintCosmetic, Category = "Input")
 	virtual void ClearAllMappings();
-
-	/**
-	 * Add a control mapping context.
-	 * 
-	 * @param MappingContext		A set of key to action mappings to apply to this player
-	 * @param Priority				Higher priority mappings will be applied first and, if they consume input, will block lower priority mappings.
-	 * @param bIgnoreAllPressedKeysUntilRelease	If true than any keys that are pressed during the rebuild of control mappings will be ignored until they are released.
-	 */
-	UE_DEPRECATED(5.0, "This version of AddMappingContext has been deprecated, please use the version that takes in a FModifyContextOptions instead.")
-	virtual void AddMappingContext(const UInputMappingContext* MappingContext, int32 Priority, const bool bIgnoreAllPressedKeysUntilRelease);
 	
 	/**
 	 * Add a control mapping context.
@@ -115,15 +105,6 @@ public:
 	 */
 	UFUNCTION(BlueprintCallable, BlueprintCosmetic, Category = "Input", meta=(AutoCreateRefTerm = "Options"))
 	virtual void AddMappingContext(const UInputMappingContext* MappingContext, int32 Priority, const FModifyContextOptions& Options = FModifyContextOptions());
-	
-	/**
-	 * Remove a specific control context. 
-	 * This is safe to call even if the context is not applied.
-	 * @param MappingContext		Context to remove from the player
-	 * @param bIgnoreAllPressedKeysUntilRelease	If true than any keys that are pressed during the rebuild of control mappings will be ignored until they are released.
-	 */
-	UE_DEPRECATED(5.0, "This version of RemoveMappingContext has been deprecated, please use the version that takes in a FModifyContextOptions instead.")
-	virtual void RemoveMappingContext(const UInputMappingContext* MappingContext, const bool bIgnoreAllPressedKeysUntilRelease);
 
 	/**
 	* Remove a specific control context. 
@@ -133,14 +114,6 @@ public:
 	*/
 	UFUNCTION(BlueprintCallable, BlueprintCosmetic, Category = "Input", meta=(AutoCreateRefTerm = "Options"))
 	virtual void RemoveMappingContext(const UInputMappingContext* MappingContext, const FModifyContextOptions& Options = FModifyContextOptions());
-	
-	/**
-	 * Flag player for reapplication of all mapping contexts at the end of this frame.
-	 * This is called automatically when adding or removing mappings contexts.
-	 * @param bForceImmediately		THe mapping changes will be applied synchronously, rather than at the end of the frame, making them available to the input system on the same frame.
-	 */
-	UE_DEPRECATED(5.0, "This version of RequestRebuildControlMappings has been deprecated, please use the version that takes in a FModifyContextOptions instead.")
-	virtual void RequestRebuildControlMappings(bool bForceImmediately, const bool bIgnoreAllPressedKeysUntilRelease = true);
 
 	/**
 	* Flag player for reapplication of all mapping contexts at the end of this frame.
