Input triggers ticking every frame checks the Input Action triggers as well

diff --git a/Engine/Plugins/EnhancedInput/Source/EnhancedInput/Private/EnhancedPlayerInput.cpp b/Engine/Plugins/EnhancedInput/Source/EnhancedInput/Private/EnhancedPlayerInput.cpp
--- a/Engine/Plugins/EnhancedInput/Source/EnhancedInput/Private/EnhancedPlayerInput.cpp
+++ b/Engine/Plugins/EnhancedInput/Source/EnhancedInput/Private/EnhancedPlayerInput.cpp
@@ -119,11 +119,21 @@ void UEnhancedPlayerInput::ProcessActionMappingEvent(const UInputAction* Action,
 	bool bMappingTriggersApplied = false;
 
 	bool bHasAnyAlwaysTickTriggers = false;
+	// checking the input mapping context triggers for any triggers that should tick every frame
 	for (const UInputTrigger* Trigger : Triggers)
 	{
-		if (Trigger)
+		if (Trigger && Trigger->bShouldAlwaysTick)
 		{
-			if (Trigger->bShouldAlwaysTick)
+			bHasAnyAlwaysTickTriggers = true;
+			break;
+		}
+	}
+	// we also need to check the triggers of the Input Action itself - only if we haven't already found an AlwaysTickTrigger
+	if (!bHasAnyAlwaysTickTriggers)
+	{
+		for (const UInputTrigger* Trigger : Action->Triggers)
+		{
+			if (Trigger && Trigger->bShouldAlwaysTick)
 			{
 				bHasAnyAlwaysTickTriggers = true;
 				break;
