Fixing an issue with combo actions not canceling the combo if pressed out of order

diff --git a/Engine/Plugins/EnhancedInput/Source/EnhancedInput/Private/InputTriggers.cpp b/Engine/Plugins/EnhancedInput/Source/EnhancedInput/Private/InputTriggers.cpp
--- a/Engine/Plugins/EnhancedInput/Source/EnhancedInput/Private/InputTriggers.cpp
+++ b/Engine/Plugins/EnhancedInput/Source/EnhancedInput/Private/InputTriggers.cpp
@@ -224,6 +224,7 @@ ETriggerState UInputTriggerComboAction::UpdateState_Implementation(const UEnhanc
 	
 	if (const UInputAction* CurrentAction = ComboActions[CurrentComboStepIndex].ComboStepAction)
 	{
+		// loop through all cancel actions and check if they've fired
 		for (const UInputAction* CancelAction : CancelActions)
 		{
 			if (CancelAction && CancelAction != CurrentAction)
@@ -234,11 +235,27 @@ ETriggerState UInputTriggerComboAction::UpdateState_Implementation(const UEnhanc
 					// Cancel action firing!
 					CurrentComboStepIndex = 0;
 					CurrentAction = ComboActions[CurrentComboStepIndex].ComboStepAction;	// Reset for fallthrough
+					break;
+				}
+			}
+		}
+		// loop through all combo actions and check if a combo action fired out of order
+		for (FInputComboStepData ComboStep : ComboActions)
+		{
+			if (ComboStep.ComboStepAction && ComboStep.ComboStepAction != CurrentAction)
+			{
+				const FInputActionInstance* CancelState = PlayerInput->FindActionInstanceData(ComboStep.ComboStepAction);
+				if (CancelState && CancelState->GetTriggerEvent() != ETriggerEvent::None)
+				{
+					// Other combo action firing - should cancel
+					CurrentComboStepIndex = 0;
+					CurrentAction = ComboActions[CurrentComboStepIndex].ComboStepAction;	// Reset for fallthrough
+					break;
 				}
 			}
 		}
 
-		// Reset if we take too long to hit the next action
+		// Reset if we take too long to hit the action
 		if (CurrentComboStepIndex > 0)
 		{
 			CurrentTimeBetweenComboSteps += DeltaTime;
@@ -250,10 +267,12 @@ ETriggerState UInputTriggerComboAction::UpdateState_Implementation(const UEnhanc
 		}
 
 		const FInputActionInstance* CurrentState = PlayerInput->FindActionInstanceData(CurrentAction);
+		// check to see if current action is completed - if so advance the combo to the next combo action
 		if (CurrentState && CurrentState->GetTriggerEvent() == ETriggerEvent::Completed) // + possibly Triggered
 		{
 			CurrentComboStepIndex++;
 			CurrentTimeBetweenComboSteps = 0;
+			// check to see if we've completed all actions in the combo
 			if (CurrentComboStepIndex >= ComboActions.Num())
 			{
 				CurrentComboStepIndex = 0;
