Add data validation to Enhanced Input asset types.

#rnx

diff --git a/Engine/Plugins/EnhancedInput/Source/EnhancedInput/Private/InputAction.cpp b/Engine/Plugins/EnhancedInput/Source/EnhancedInput/Private/InputAction.cpp
--- a/Engine/Plugins/EnhancedInput/Source/EnhancedInput/Private/InputAction.cpp
+++ b/Engine/Plugins/EnhancedInput/Source/EnhancedInput/Private/InputAction.cpp
@@ -7,6 +7,8 @@
 #include "InputTriggers.h"
 #include "UObject/UObjectIterator.h"
 
+#define LOCTEXT_NAMESPACE "EnhancedInputAction"
+
 FInputActionInstance::FInputActionInstance(const UInputAction* InSourceAction)
 	: SourceAction(InSourceAction)
 {
@@ -113,6 +115,41 @@ void UInputAction::PostEditChangeProperty(FPropertyChangedEvent& PropertyChanged
 	}
 }
 
+EDataValidationResult UInputAction::IsDataValid(TArray<FText>& ValidationErrors)
+{
+	EDataValidationResult Result = CombineDataValidationResults(Super::IsDataValid(ValidationErrors), EDataValidationResult::Valid);
+
+	// Validate the triggers
+	for (const TObjectPtr<UInputTrigger> Trigger : Triggers)
+	{
+		if (Trigger)
+		{
+			Result = CombineDataValidationResults(Result, Trigger->IsDataValid(ValidationErrors));
+		}
+		else
+		{
+			Result = EDataValidationResult::Invalid;
+			ValidationErrors.Add(LOCTEXT("NullInputTrigger", "There cannot be a null Input Trigger on an Input Action!"));
+		}
+	}
+	
+	// Validate the modifiers
+	for (const TObjectPtr<UInputModifier> Modifier : Modifiers)
+	{
+		if (Modifier)
+		{
+			Result = CombineDataValidationResults(Result, Modifier->IsDataValid(ValidationErrors));	
+		}
+		else
+		{
+			Result = EDataValidationResult::Invalid;
+			ValidationErrors.Add(LOCTEXT("NullInputModifier", "There cannot be a null Input Modifier on an Input Action!"));
+		}		
+	}
+
+	return Result;
+}
+
 ETriggerEventsSupported UInputAction::GetSupportedTriggerEvents() const
 {
 	ETriggerEventsSupported SupportedTriggers = ETriggerEventsSupported::None;
@@ -140,4 +177,6 @@ ETriggerEventsSupported UInputAction::GetSupportedTriggerEvents() const
 	return SupportedTriggers;
 }
 
-#endif
\ No newline at end of file
+#endif
+
+#undef LOCTEXT_NAMESPACE
\ No newline at end of file
diff --git a/Engine/Plugins/EnhancedInput/Source/EnhancedInput/Private/InputMappingContext.cpp b/Engine/Plugins/EnhancedInput/Source/EnhancedInput/Private/InputMappingContext.cpp
--- a/Engine/Plugins/EnhancedInput/Source/EnhancedInput/Private/InputMappingContext.cpp
+++ b/Engine/Plugins/EnhancedInput/Source/EnhancedInput/Private/InputMappingContext.cpp
@@ -5,6 +5,60 @@
 #include "EnhancedInputLibrary.h"
 #include "EnhancedInputModule.h"
 
+#define LOCTEXT_NAMESPACE "InputMappingContext"
+
+#if WITH_EDITOR
+EDataValidationResult UInputMappingContext::IsDataValid(TArray<FText>& ValidationErrors)
+{
+	EDataValidationResult Result = CombineDataValidationResults(Super::IsDataValid(ValidationErrors), EDataValidationResult::Valid);
+	
+	for (const FEnhancedActionKeyMapping& Mapping : Mappings)
+	{
+		if (!Mapping.Action)
+		{
+			Result = EDataValidationResult::Invalid;
+			ValidationErrors.Add(LOCTEXT("NullInputAction", "A mapping cannot have an empty input action!"));
+		}
+		
+		if (Mapping.bIsPlayerMappable && Mapping.PlayerMappableOptions.Name == NAME_None)
+		{
+			Result = EDataValidationResult::Invalid;
+        	ValidationErrors.Add(LOCTEXT("InvalidPlayerMappableName", "A player mappable key mapping must have a valid 'Name'"));
+		}
+		
+		// Validate the triggers
+		for (const TObjectPtr<UInputTrigger> Trigger : Mapping.Triggers)
+		{
+			if (Trigger)
+			{
+				Result = CombineDataValidationResults(Result, Trigger->IsDataValid(ValidationErrors));
+			}
+			else
+			{
+				Result = EDataValidationResult::Invalid;
+				ValidationErrors.Add(LOCTEXT("NullInputTrigger", "There cannot be a null Input Trigger on a key mapping"));
+			}
+		}
+		
+		// Validate the modifiers
+		for (const TObjectPtr<UInputModifier> Modifier : Mapping.Modifiers)
+		{
+			if (Modifier)
+			{
+				Result = CombineDataValidationResults(Result, Modifier->IsDataValid(ValidationErrors));	
+			}
+			else
+			{
+				Result = EDataValidationResult::Invalid;
+				ValidationErrors.Add(LOCTEXT("NullInputModifier", "There cannot be a null Input Modifier on a key mapping"));
+			}
+		}
+	}
+
+	return Result;
+}
+#endif	// WITH_EDITOR
+
 FEnhancedActionKeyMapping& UInputMappingContext::MapKey(const UInputAction* Action, FKey ToKey)
 {
 	IEnhancedInputModule::Get().GetLibrary()->RequestRebuildControlMappingsUsingContext(this);
@@ -37,4 +91,6 @@ void UInputMappingContext::UnmapAll()
 		Mappings.Empty();
 		IEnhancedInputModule::Get().GetLibrary()->RequestRebuildControlMappingsUsingContext(this);
 	}
-}
\ No newline at end of file
+}
+
+#undef LOCTEXT_NAMESPACE
\ No newline at end of file
diff --git a/Engine/Plugins/EnhancedInput/Source/EnhancedInput/Private/InputModifiers.cpp b/Engine/Plugins/EnhancedInput/Source/EnhancedInput/Private/InputModifiers.cpp
--- a/Engine/Plugins/EnhancedInput/Source/EnhancedInput/Private/InputModifiers.cpp
+++ b/Engine/Plugins/EnhancedInput/Source/EnhancedInput/Private/InputModifiers.cpp
@@ -6,12 +6,31 @@
 #include "GameFramework/InputSettings.h"
 #include "GameFramework/PlayerController.h"
 
-
+#define LOCTEXT_NAMESPACE "EnhancedInputModifiers"
 
 /*
 * Scalar
 */
 
+#if WITH_EDITOR
+EDataValidationResult UInputModifierScalar::IsDataValid(TArray<FText>& ValidationErrors)
+{
+	EDataValidationResult Result = CombineDataValidationResults(Super::IsDataValid(ValidationErrors), EDataValidationResult::Valid);
+
+	// You cannot scale a boolean value
+	if (UInputAction* IA = Cast<UInputAction>(GetOuter()))
+	{
+		if (IA->ValueType == EInputActionValueType::Boolean)
+		{
+			Result = EDataValidationResult::Invalid;
+			ValidationErrors.Add(LOCTEXT("InputScalarInvalidActionType", "A Scalar modifier cannot be used on a 'Boolean' input action"));
+		}
+	}
+	
+	return Result;
+}
+#endif
+
 FInputActionValue UInputModifierScalar::ModifyRaw_Implementation(const UEnhancedPlayerInput* PlayerInput, FInputActionValue CurrentValue, float DeltaTime)
 {
 	// Don't try and scale bools
@@ -346,4 +365,6 @@ FLinearColor UInputModifierSwizzleAxis::GetVisualizationColor_Implementation(FIn
 	// TODO: Color blend per swizzle type?
 	float SampleX = (FMath::Abs(SampleValue.Get<float>()) + 1.f) * 0.5f;
 	return FLinearColor(SampleX, 1.f - SampleX, 0.f);
-}
\ No newline at end of file
+}
+
+#undef LOCTEXT_NAMESPACE
\ No newline at end of file
diff --git a/Engine/Plugins/EnhancedInput/Source/EnhancedInput/Private/InputTriggers.cpp b/Engine/Plugins/EnhancedInput/Source/EnhancedInput/Private/InputTriggers.cpp
--- a/Engine/Plugins/EnhancedInput/Source/EnhancedInput/Private/InputTriggers.cpp
+++ b/Engine/Plugins/EnhancedInput/Source/EnhancedInput/Private/InputTriggers.cpp
@@ -6,6 +6,8 @@
 #include "EnhancedPlayerInput.h"
 #include "GameFramework/PlayerController.h"
 
+#define LOCTEXT_NAMESPACE "EnhancedInputTriggers"
+
 // Abstract trigger bases
 ETriggerState UInputTrigger::UpdateState_Implementation(const UEnhancedPlayerInput* PlayerInput, FInputActionValue ModifiedValue, float DeltaTime)
 {
@@ -201,6 +203,21 @@ ETriggerState UInputTriggerPulse::UpdateState_Implementation(const UEnhancedPlay
 	return State;
 }
 
+#if WITH_EDITOR
+EDataValidationResult UInputTriggerChordAction::IsDataValid(TArray<FText>& ValidationErrors)
+{
+	EDataValidationResult Result = CombineDataValidationResults(Super::IsDataValid(ValidationErrors), EDataValidationResult::Valid);
+	
+	// You can't evaluate the combo if there are no combo steps!
+	if (!ChordAction)
+	{
+		Result = EDataValidationResult::Invalid;
+		ValidationErrors.Add(LOCTEXT("NullChordedAction", "A valid action is required for the Chorded Action input trigger!"));
+	}
+
+	return Result;
+}
+#endif
 
 ETriggerState UInputTriggerChordAction::UpdateState_Implementation(const UEnhancedPlayerInput* PlayerInput, FInputActionValue ModifiedValue, float DeltaTime)
 {
@@ -294,4 +311,23 @@ ETriggerState UInputTriggerComboAction::UpdateState_Implementation(const UEnhanc
 		CurrentTimeBetweenComboSteps = 0;
 	}
 	return ETriggerState::None;
-};
\ No newline at end of file
+};
+
+#if WITH_EDITOR
+EDataValidationResult UInputTriggerComboAction::IsDataValid(TArray<FText>& ValidationErrors)
+{
+	EDataValidationResult Result = CombineDataValidationResults(Super::IsDataValid(ValidationErrors), EDataValidationResult::Valid);
+	
+	// You can't evaluate the combo if there are no combo steps!
+	if (ComboActions.IsEmpty())
+	{
+		Result = EDataValidationResult::Invalid;
+		ValidationErrors.Add(LOCTEXT("NoComboSteps", "There must be at least one combo step in the Combo Trigger!"));
+	}
+
+	return Result;
+}
+
+#endif
+
+#undef LOCTEXT_NAMESPACE
\ No newline at end of file
diff --git a/Engine/Plugins/EnhancedInput/Source/EnhancedInput/Public/InputAction.h b/Engine/Plugins/EnhancedInput/Source/EnhancedInput/Public/InputAction.h
--- a/Engine/Plugins/EnhancedInput/Source/EnhancedInput/Public/InputAction.h
+++ b/Engine/Plugins/EnhancedInput/Source/EnhancedInput/Public/InputAction.h
@@ -24,6 +24,8 @@ public:
 #if WITH_EDITOR
 	// Track actions that have had their ValueType changed to update blueprints referencing them.
 	virtual void PostEditChangeProperty(FPropertyChangedEvent& PropertyChangedEvent) override;
+	virtual EDataValidationResult IsDataValid(TArray<FText>& ValidationErrors) override;
+	
 	static TSet<const UInputAction*> ActionsWithModifiedValueTypes;
 	
 	/**
diff --git a/Engine/Plugins/EnhancedInput/Source/EnhancedInput/Public/InputMappingContext.h b/Engine/Plugins/EnhancedInput/Source/EnhancedInput/Public/InputMappingContext.h
--- a/Engine/Plugins/EnhancedInput/Source/EnhancedInput/Public/InputMappingContext.h
+++ b/Engine/Plugins/EnhancedInput/Source/EnhancedInput/Public/InputMappingContext.h
@@ -30,6 +30,10 @@ protected:
 
 public:
 
+#if WITH_EDITOR
+	virtual EDataValidationResult IsDataValid(TArray<FText>& ValidationErrors) override;
+#endif
+
 	// Localized context descriptor
 	UPROPERTY(BlueprintReadOnly, EditAnywhere, Category = "Description", DisplayName = "Description")
 	FText ContextDescription;
diff --git a/Engine/Plugins/EnhancedInput/Source/EnhancedInput/Public/InputModifiers.h b/Engine/Plugins/EnhancedInput/Source/EnhancedInput/Public/InputModifiers.h
--- a/Engine/Plugins/EnhancedInput/Source/EnhancedInput/Public/InputModifiers.h
+++ b/Engine/Plugins/EnhancedInput/Source/EnhancedInput/Public/InputModifiers.h
@@ -102,6 +102,10 @@ class UInputModifierScalar : public UInputModifier
 
 public:
 
+#if WITH_EDITOR
+	virtual EDataValidationResult IsDataValid(TArray<FText>& ValidationErrors) override;
+#endif
+
 	// TODO: Detail customization to only show modifiable axes for the relevant binding? This thing has no idea what it's bound to...
 	UPROPERTY(EditInstanceOnly, BlueprintReadWrite, Category=Settings)
 	FVector Scalar = FVector::OneVector;
diff --git a/Engine/Plugins/EnhancedInput/Source/EnhancedInput/Public/InputTriggers.h b/Engine/Plugins/EnhancedInput/Source/EnhancedInput/Public/InputTriggers.h
--- a/Engine/Plugins/EnhancedInput/Source/EnhancedInput/Public/InputTriggers.h
+++ b/Engine/Plugins/EnhancedInput/Source/EnhancedInput/Public/InputTriggers.h
@@ -384,6 +384,10 @@ UCLASS(NotBlueprintable, meta = (DisplayName = "Chorded Action", NotInputConfigu
 class ENHANCEDINPUT_API UInputTriggerChordAction : public UInputTrigger
 {
 	GENERATED_BODY()
+public:
+#if WITH_EDITOR
+	virtual EDataValidationResult IsDataValid(TArray<FText>& ValidationErrors) override;
+#endif
 
 protected:
 	// Implicit, so action cannot fire unless this is firing.
@@ -449,6 +453,11 @@ protected:
 	float CurrentTimeBetweenComboSteps = 0.0f;
 
 public:
+
+#if WITH_EDITOR
+	virtual EDataValidationResult IsDataValid(TArray<FText>& ValidationErrors) override;
+#endif
+	
 	/**
 	 * List of input actions that need to be completed to trigger this action.
 	 * Input actions must be triggered in order (starting at index 0) to count towards the completion of the combo.
