Add the Enhanced Input plugin icon!

#rnx

diff --git a/Engine/Plugins/EnhancedInput/Source/InputEditor/Private/InputEditorModule.cpp b/Engine/Plugins/EnhancedInput/Source/InputEditor/Private/InputEditorModule.cpp
--- a/Engine/Plugins/EnhancedInput/Source/InputEditor/Private/InputEditorModule.cpp
+++ b/Engine/Plugins/EnhancedInput/Source/InputEditor/Private/InputEditorModule.cpp
@@ -169,6 +169,9 @@ public:
 		static const FVector2D Icon16 = FVector2D(16.0f, 16.0f);
 		static const FVector2D Icon64 = FVector2D(64.0f, 64.0f);
 
+		Set("EnhancedInputIcon_Small", new IMAGE_BRUSH_SVG("Icons/EnhancedInput_16", Icon16));
+		Set("EnhancedInputIcon_Large", new IMAGE_BRUSH_SVG("Icons/EnhancedInput_64", Icon64));
+		
 		Set("ClassIcon.InputAction", new IMAGE_BRUSH_SVG("Icons/InputAction_16", Icon16));
 		Set("ClassThumbnail.InputAction", new IMAGE_BRUSH_SVG("Icons/InputAction_64", Icon64));
 		
