Change bIgnoreAllPressedKeysUntilReleaseOnRebuild ignore "boolean" input action types instead of Digital FKey types. This way keys that are a part of analog actions behave as if they were analog buttons when considering the bIgnoreAllPressedKeysUntilReleaseOnRebuild  flag.

#rnx

diff --git a/Engine/Plugins/EnhancedInput/Source/EnhancedInput/Private/EnhancedInputSubsystemInterface.cpp b/Engine/Plugins/EnhancedInput/Source/EnhancedInput/Private/EnhancedInputSubsystemInterface.cpp
--- a/Engine/Plugins/EnhancedInput/Source/EnhancedInput/Private/EnhancedInputSubsystemInterface.cpp
+++ b/Engine/Plugins/EnhancedInput/Source/EnhancedInput/Private/EnhancedInputSubsystemInterface.cpp
@@ -704,7 +704,7 @@ void IEnhancedInputSubsystemInterface::RebuildControlMappings()
 			// Was this key pressed last frame? If so, then we need to mark it to be ignored by the PlayerInput
 			// until it is released to avoid re-processing a triggered event when it. This is only a problem if
 			// the key was in the old mapping and the new one
-			if(bIgnoreAllPressedKeysUntilReleaseOnRebuild && Mapping.Key.IsDigital() && WasInOldMapping(Mapping.Key))
+			if(bIgnoreAllPressedKeysUntilReleaseOnRebuild && Mapping.Action->ValueType == EInputActionValueType::Boolean && WasInOldMapping(Mapping.Key))
 			{				
 				const FKeyState* KeyState = PlayerInput->GetKeyState(Mapping.Key);
 				if(KeyState && KeyState->bDown)
