PR #9435: [EnhancedInput] Fix for BindAction() not accepting additional function parameters. (Contributed by [REDACTED_PERSON])

diff --git a/Engine/Plugins/EnhancedInput/Source/EnhancedInput/Public/EnhancedInputComponent.h b/Engine/Plugins/EnhancedInput/Source/EnhancedInput/Public/EnhancedInputComponent.h
--- a/Engine/Plugins/EnhancedInput/Source/EnhancedInput/Public/EnhancedInputComponent.h
+++ b/Engine/Plugins/EnhancedInput/Source/EnhancedInput/Public/EnhancedInputComponent.h
@@ -72,11 +72,12 @@ public:
 
 	/** Binds a native delegate, hidden for script delegates */
 	template<	typename UserClass,
-				typename TSig = TSignature>
-	void BindDelegate(UserClass* Object, typename TSig::template TMethodPtr<UserClass> Func)
+				typename TSig = TSignature,
+				typename... TVars>
+	void BindDelegate(UserClass* Object, typename TSig::template TMethodPtr<UserClass, TVars...> Func, TVars... Vars)
 	{
 		Unbind();
-		Delegate = MakeShared<TSig>(TSig::CreateUObject(Object, Func));
+		Delegate = MakeShared<TSig>(TSig::CreateUObject(Object, Func, Vars...));
 	}
 
 	/** Binds a script delegate on an arbitrary UObject, hidden for native delegates */
@@ -408,11 +409,11 @@ public:
 	 * Binds a delegate function matching any of the handler signatures to a UInputAction assigned via UInputMappingContext to the owner of this component.
 	 */
 #define DEFINE_BIND_ACTION(HANDLER_SIG)																																			\
-	template<class UserClass>																																					\
-	FEnhancedInputActionEventBinding& BindAction(const UInputAction* Action, ETriggerEvent TriggerEvent, UserClass* Object, typename HANDLER_SIG::TMethodPtr< UserClass > Func) \
+	template<class UserClass, typename... VarTypes>																																					\
+	FEnhancedInputActionEventBinding& BindAction(const UInputAction* Action, ETriggerEvent TriggerEvent, UserClass* Object, typename HANDLER_SIG::template TMethodPtr< UserClass, VarTypes... > Func, VarTypes... Vars) \
 	{																																											\
 		TUniquePtr<FEnhancedInputActionEventDelegateBinding<HANDLER_SIG>> AB = MakeUnique<FEnhancedInputActionEventDelegateBinding<HANDLER_SIG>>(Action, TriggerEvent);			\
-		AB->Delegate.BindDelegate<UserClass>(Object, Func);																														\
+		AB->Delegate.BindDelegate<UserClass>(Object, Func, Vars...);																														\
 		AB->Delegate.SetShouldFireWithEditorScriptGuard(bShouldFireDelegatesInEditor);																							\
 		return *EnhancedActionEventBindings.Add_GetRef(MoveTemp(AB));																											\
 	}
@@ -421,15 +422,6 @@ public:
 	DEFINE_BIND_ACTION(FEnhancedInputActionHandlerValueSignature);
 	DEFINE_BIND_ACTION(FEnhancedInputActionHandlerInstanceSignature);
 
-	template< class FuncType, class UserClass, typename... VarTypes >
-	FEnhancedInputActionEventBinding& BindAction(const UInputAction* Action, ETriggerEvent TriggerEvent, UserClass* Object, FuncType Func, VarTypes... Vars)
-	{
-		TUniquePtr<FEnhancedInputActionEventDelegateBinding<FEnhancedInputActionHandlerSignature>> AB = MakeUnique<FEnhancedInputActionEventDelegateBinding<FEnhancedInputActionHandlerSignature>>(Action, TriggerEvent);
-		AB->Delegate.MakeDelegate().BindUObject(Object, Func, Vars...);
-		AB->Delegate.SetShouldFireWithEditorScriptGuard(bShouldFireDelegatesInEditor);
-		return *EnhancedActionEventBindings.Add_GetRef(MoveTemp(AB));
-	}
-
 	/**
 	 * Binds to an object UFUNCTION
 	 */
