Rename Enhanced Input BaseEngine.ini file to the BaseEnhancedInput.ini.

#rnx

diff --git a/Engine/Plugins/EnhancedInput/Config/BaseEngine.ini b/Engine/Plugins/EnhancedInput/Config/BaseEngine.ini
deleted file mode 100644
diff --git a/Engine/Plugins/EnhancedInput/Config/BaseEnhancedInput.ini b/Engine/Plugins/EnhancedInput/Config/BaseEnhancedInput.ini
new file mode 100644
--- /dev/null
+++ b/Engine/Plugins/EnhancedInput/Config/BaseEnhancedInput.ini
@@ -0,0 +1,2 @@
+[CoreRedirects]
++ClassRedirects=(OldName="/Script/InputMappingContext.UnmapAction", NewName="/Script/InputMappingContext.UnmapAllKeysFromAction")
\ No newline at end of file
