Add a right click option to create a Mapping Context from the selected input actions. This will let you create an input mapping context that is pre-populated with some input actions!

#rnx

diff --git a/Engine/Plugins/EnhancedInput/Source/InputEditor/InputEditor.Build.cs b/Engine/Plugins/EnhancedInput/Source/InputEditor/InputEditor.Build.cs
--- a/Engine/Plugins/EnhancedInput/Source/InputEditor/InputEditor.Build.cs
+++ b/Engine/Plugins/EnhancedInput/Source/InputEditor/InputEditor.Build.cs
@@ -37,6 +37,8 @@ namespace UnrealBuildTool.Rules
 					"AssetTools",
 					"DeveloperSettings",
 					"EditorSubsystem",
+					"ToolMenus",
+					"ContentBrowser",
 					"SourceControl",
                 }
             );
diff --git a/Engine/Plugins/EnhancedInput/Source/InputEditor/Private/InputEditorModule.cpp b/Engine/Plugins/EnhancedInput/Source/InputEditor/Private/InputEditorModule.cpp
--- a/Engine/Plugins/EnhancedInput/Source/InputEditor/Private/InputEditorModule.cpp
+++ b/Engine/Plugins/EnhancedInput/Source/InputEditor/Private/InputEditorModule.cpp
@@ -15,6 +15,7 @@
 #include "InputCustomizations.h"
 #include "InputModifiers.h"
 #include "ISettingsModule.h"
+#include "ToolMenuSection.h"
 #include "K2Node_EnhancedInputAction.h"
 #include "K2Node_GetInputActionValue.h"
 #include "Modules/ModuleInterface.h"
@@ -30,6 +31,8 @@
 #include "Styling/StyleColors.h"
 #include "Styling/SlateStyleMacros.h"
 #include "Styling/SlateStyleRegistry.h"
+#include "ContentBrowserModule.h"
+#include "IContentBrowserSingleton.h"
 #include "GameFramework/InputSettings.h"
 #include "EnhancedInputComponent.h"
 #include "EnhancedPlayerInput.h"
@@ -67,7 +70,24 @@ UInputMappingContext_Factory::UInputMappingContext_Factory(const class FObjectIn
 UObject* UInputMappingContext_Factory::FactoryCreateNew(UClass* Class, UObject* InParent, FName Name, EObjectFlags Flags, UObject* Context, FFeedbackContext* Warn)
 {
 	check(Class->IsChildOf(UInputMappingContext::StaticClass()));
-	return NewObject<UInputMappingContext>(InParent, Class, Name, Flags | RF_Transactional, Context);
+	UInputMappingContext* IMC = NewObject<UInputMappingContext>(InParent, Class, Name, Flags | RF_Transactional, Context);
+
+	// Populate the IMC with some initial input actions if they were specified. This will be the case if the IMC is being created from the FAssetTypeActions_InputAction
+	for (TWeakObjectPtr<UInputAction> WeakIA : InitialActions)
+	{
+		if (UInputAction* IA = WeakIA.Get())
+		{
+			IMC->MapKey(IA, FKey());
+		}
+	}
+	
+	return IMC;
+}
+
+void UInputMappingContext_Factory::SetInitialActions(TArray<TWeakObjectPtr<UInputAction>> InInitialActions)
+{
+	InitialActions.Empty();
+	InitialActions = InInitialActions;
 }
 
 // InputAction
@@ -100,27 +120,7 @@ UObject* UPlayerMappableInputConfig_Factory::FactoryCreateNew(UClass* Class, UOb
 	return NewObject<UPlayerMappableInputConfig>(InParent, Class, Name, Flags | RF_Transactional, Context);
 }
 
-//
-//// InputTrigger
-//UInputTrigger_Factory::UInputTrigger_Factory(const class FObjectInitializer& OBJ) : Super(OBJ) {
-//	ParentClass = UInputTrigger::StaticClass();
-//	SupportedClass = UInputTrigger::StaticClass();
-//	bEditAfterNew = true;
-//	bCreateNew = true;
-//}
-//
-//// InputModifier
-//UInputModifier_Factory::UInputModifier_Factory(const class FObjectInitializer& OBJ) : Super(OBJ) {
-//	ParentClass = UInputModifier::StaticClass();
-//	SupportedClass = UInputModifier::StaticClass();
-//	bEditAfterNew = true;
-//	bCreateNew = true;
-//}
-
-
-
 // Asset type actions
-// TODO: Move asset type action definitions out?
 
 class FAssetTypeActions_InputContext : public FAssetTypeActions_DataAsset
 {
@@ -140,8 +140,66 @@ public:
 	virtual FColor GetTypeColor() const override { return FColor(127, 255, 255); }
 	virtual FText GetAssetDescription(const FAssetData& AssetData) const override { return NSLOCTEXT("AssetTypeActions", "AssetTypeActions_InputActionDesc", "Represents an an abstract game action that can be mapped to arbitrary hardware input devices."); }
 	virtual UClass* GetSupportedClass() const override { return UInputAction::StaticClass(); }
+	virtual bool HasActions(const TArray<UObject*>& InObjects) const override { return true; }
+	virtual void GetActions(const TArray<UObject*>& InObjects, struct FToolMenuSection& Section) override;
+
+private:
+
+	void GetCreateContextFromActionsMenu(TArray<TWeakObjectPtr<UInputAction>> InActions);
+	
 };
 
+void FAssetTypeActions_InputAction::GetActions(const TArray<UObject*>& InObjects, FToolMenuSection& Section)
+{
+	TArray<TWeakObjectPtr<UInputAction>> InputActions = GetTypedWeakObjectPtrs<UInputAction>(InObjects);
+	
+	 Section.AddMenuEntry(
+	 	"InputAction_CreateContextFromSelection",	
+	 	LOCTEXT("InputAction_CreateContextFromSelection", "Create an Input Mapping Context"),
+	 	LOCTEXT("InputAction_CreateContextFromSelectionTooltip", "Create an Input Mapping Context that is filled with the selected Input Actions"),
+	 	FSlateIcon(FAppStyle::GetAppStyleSetName(), "ClassIcon.InputMappingContext"),
+	 	FUIAction(FExecuteAction::CreateSP(this, &FAssetTypeActions_InputAction::GetCreateContextFromActionsMenu, InputActions))
+	 );
+}
+
+void FAssetTypeActions_InputAction::GetCreateContextFromActionsMenu(TArray<TWeakObjectPtr<UInputAction>> InActions)
+{
+	if (InActions.IsEmpty())
+	{
+		return;
+	}
+
+	static const FString MappingContextPrefix(TEXT("IMC_"));
+	static const FString ActionPrefixToIgnore(TEXT("IA_"));
+	
+	if (UInputAction* FirstIA = InActions[0].Get())
+	{
+		FString EffectiveIMCName = FirstIA->GetName();
+		EffectiveIMCName.RemoveFromStart(ActionPrefixToIgnore);
+		// Have the IMC_ prefix at the front of the new asset
+		EffectiveIMCName.InsertAt(0, MappingContextPrefix);
+		
+		const FString ActionPathName = FirstIA->GetOutermost()->GetPathName();
+    	const FString LongPackagePath = FPackageName::GetLongPackagePath(ActionPathName);
+		const FString NewIMCDefaultPath = LongPackagePath / EffectiveIMCName;
+
+		FAssetToolsModule& AssetToolsModule = FModuleManager::GetModuleChecked<FAssetToolsModule>("AssetTools");
+		FContentBrowserModule& ContentBrowserModule = FModuleManager::LoadModuleChecked<FContentBrowserModule>("ContentBrowser");
+		
+		// Make sure the name is unique
+		FString AssetName;
+		FString DefaultSuffix;
+		FString PackageName;
+		AssetToolsModule.Get().CreateUniqueAssetName(NewIMCDefaultPath, DefaultSuffix, /*out*/ PackageName, /*out*/ AssetName);
+		const FString PackagePath = FPackageName::GetLongPackagePath(PackageName);
+
+		// Create the new IMC
+		UInputMappingContext_Factory* IMCFactory = NewObject<UInputMappingContext_Factory>();
+		IMCFactory->SetInitialActions(InActions);
+		ContentBrowserModule.Get().CreateNewAsset(AssetName, PackagePath, UInputMappingContext::StaticClass(), IMCFactory);
+	}
+}
+
 class FAssetTypeActions_PlayerMappableInputConfig : public FAssetTypeActions_DataAsset
 {
 public:
diff --git a/Engine/Plugins/EnhancedInput/Source/InputEditor/Public/InputEditorModule.h b/Engine/Plugins/EnhancedInput/Source/InputEditor/Public/InputEditorModule.h
--- a/Engine/Plugins/EnhancedInput/Source/InputEditor/Public/InputEditorModule.h
+++ b/Engine/Plugins/EnhancedInput/Source/InputEditor/Public/InputEditorModule.h
@@ -17,6 +17,7 @@ DECLARE_LOG_CATEGORY_EXTERN(LogEnhancedInputEditor, Log, All);
 ////////////////////////////////////////////////////////////////////
 // FInputEditorModule
 
+class UInputAction;
 class SWindow;
 
 class FInputEditorModule : public IModuleInterface, public FTickableEditorObject
@@ -63,6 +64,15 @@ class INPUTEDITOR_API UInputMappingContext_Factory : public UFactory
 	GENERATED_UCLASS_BODY()
 public:
 	virtual UObject* FactoryCreateNew(UClass* Class, UObject* InParent, FName Name, EObjectFlags Flags, UObject* Context, FFeedbackContext* Warn) override;
+
+	/** Set the array of initial actions that the resulting IMC should be populated with */
+	void SetInitialActions(TArray<TWeakObjectPtr<UInputAction>> InInitialActions);
+	
+protected:
+
+	/** An array of Input Actions that the mapping context should be populated with upon creation */
+	TArray<TWeakObjectPtr<UInputAction>> InitialActions;
+	
 };
 
 UCLASS()
