Add a functoin to allow you to get the mapping context priority from the Enhanced Input subsystem.

diff --git a/Engine/Plugins/EnhancedInput/Source/EnhancedInput/Private/EnhancedInputSubsystemInterface.cpp b/Engine/Plugins/EnhancedInput/Source/EnhancedInput/Private/EnhancedInputSubsystemInterface.cpp
--- a/Engine/Plugins/EnhancedInput/Source/EnhancedInput/Private/EnhancedInputSubsystemInterface.cpp
+++ b/Engine/Plugins/EnhancedInput/Source/EnhancedInput/Private/EnhancedInputSubsystemInterface.cpp
@@ -335,10 +335,27 @@ void IEnhancedInputSubsystemInterface::ApplyAxisPropertyModifiers(UEnhancedPlaye
 	}
 }
 
-
 bool IEnhancedInputSubsystemInterface::HasMappingContext(const UInputMappingContext* MappingContext) const
 {
-	return GetPlayerInput() && GetPlayerInput()->AppliedInputContexts.Contains(MappingContext);
+	int32 DummyPri = INDEX_NONE;
+	return HasMappingContext(MappingContext, DummyPri);
+}
+
+bool IEnhancedInputSubsystemInterface::HasMappingContext(const UInputMappingContext* MappingContext, int32& OutFoundPriority) const
+{
+	bool bResult = false;
+	OutFoundPriority = INDEX_NONE;
+	
+	if (const UEnhancedPlayerInput* const Input = GetPlayerInput())
+	{
+		if (const int32* const FoundPriority = Input->AppliedInputContexts.Find(MappingContext))
+		{
+			OutFoundPriority = *FoundPriority;
+			bResult = true;
+		}
+	}
+
+	return bResult;
 }
 
 TArray<FKey> IEnhancedInputSubsystemInterface::QueryKeysMappedToAction(const UInputAction* Action) const
diff --git a/Engine/Plugins/EnhancedInput/Source/EnhancedInput/Public/EnhancedInputSubsystemInterface.h b/Engine/Plugins/EnhancedInput/Source/EnhancedInput/Public/EnhancedInputSubsystemInterface.h
--- a/Engine/Plugins/EnhancedInput/Source/EnhancedInput/Public/EnhancedInputSubsystemInterface.h
+++ b/Engine/Plugins/EnhancedInput/Source/EnhancedInput/Public/EnhancedInputSubsystemInterface.h
@@ -154,9 +154,18 @@ public:
 	/**
 	 * Check if a mapping context is applied to this subsystem's owner.
 	 */
-	UFUNCTION(BlueprintCallable, BlueprintCosmetic, Category = "Input|Mapping Queries")	// TODO: BlueprintPure would be nicer. Move into library?
 	virtual bool HasMappingContext(const UInputMappingContext* MappingContext) const;
 
+	/**
+	 * Check if a mapping context is applied to this subsystem's owner.
+	 *
+	 * @param MappingContext		The mapping context to search for on the subsystem's owner.
+	 * @param OutFoundPriority		The priority of the mapping context if it is applied. -1 if the context is not applied	
+	 * @return	True if the mapping context is applied
+	 */
+	UFUNCTION(BlueprintCallable, BlueprintCosmetic, Category = "Input|Mapping Queries")
+	virtual bool HasMappingContext(const UInputMappingContext* MappingContext, int32& OutFoundPriority) const;
+
 	/**
 	 * Returns the keys mapped to the given action in the active input mapping contexts.
 	 */
