Combo Input Triggers now return the correct supported trigger events.
Adding a config variable for Unsupported triggers warnings.

diff --git a/Engine/Plugins/EnhancedInput/Source/EnhancedInput/Public/InputTriggers.h b/Engine/Plugins/EnhancedInput/Source/EnhancedInput/Public/InputTriggers.h
--- a/Engine/Plugins/EnhancedInput/Source/EnhancedInput/Public/InputTriggers.h
+++ b/Engine/Plugins/EnhancedInput/Source/EnhancedInput/Public/InputTriggers.h
@@ -468,4 +468,7 @@ public:
 	// Actions that will cancel the combo if they are triggered
     UPROPERTY(EditAnywhere, BlueprintReadWrite, Category = "Trigger Settings", meta = (DisplayThumbnail = "false"))
     TArray<TObjectPtr<const UInputAction>> CancelActions;
+
+	/** Determines what kind of trigger events can happen from the behavior of this trigger. */
+	virtual ETriggerEventsSupported GetSupportedTriggerEvents() const override { return ETriggerEventsSupported::All; }
 };
\ No newline at end of file
diff --git a/Engine/Plugins/EnhancedInput/Source/InputBlueprintNodes/Private/K2Node_EnhancedInputAction.cpp b/Engine/Plugins/EnhancedInput/Source/InputBlueprintNodes/Private/K2Node_EnhancedInputAction.cpp
--- a/Engine/Plugins/EnhancedInput/Source/InputBlueprintNodes/Private/K2Node_EnhancedInputAction.cpp
+++ b/Engine/Plugins/EnhancedInput/Source/InputBlueprintNodes/Private/K2Node_EnhancedInputAction.cpp
@@ -4,6 +4,7 @@
 #include "AssetRegistry/AssetRegistryModule.h"
 #include "BlueprintActionDatabase.h"
 #include "BlueprintActionDatabaseRegistrar.h"
+#include "BlueprintEditorSettings.h"
 #include "BlueprintNodeSpawner.h"
 #include "EdGraphSchema_K2.h"
 #include "EdGraphSchema_K2_Actions.h"
@@ -60,7 +61,7 @@ void UK2Node_EnhancedInputAction::AllocateDefaultPins()
 {
 	PreloadObject((UObject*)InputAction);
 
-	const ETriggerEventsSupported SupportedTriggerEvents = InputAction ? InputAction->GetSupportedTriggerEvents() : ETriggerEventsSupported::None;
+	const ETriggerEventsSupported SupportedTriggerEvents = GetDefault<UBlueprintEditorSettings>()->bEnableInputTriggerSupportWarnings && InputAction ? InputAction->GetSupportedTriggerEvents() : ETriggerEventsSupported::All;
 	
 	ForEachEventPinName([this, SupportedTriggerEvents](ETriggerEvent Event, FName PinName)
 	{
@@ -104,7 +105,7 @@ void UK2Node_EnhancedInputAction::AllocateDefaultPins()
 void UK2Node_EnhancedInputAction::HideEventPins(UEdGraphPin* RetainPin)
 {
 	// Gather pins
-	const ETriggerEventsSupported SupportedTriggerEvents = InputAction ? InputAction->GetSupportedTriggerEvents() : ETriggerEventsSupported::None;
+	const ETriggerEventsSupported SupportedTriggerEvents = GetDefault<UBlueprintEditorSettings>()->bEnableInputTriggerSupportWarnings && InputAction ? InputAction->GetSupportedTriggerEvents() : ETriggerEventsSupported::All;
 
 	// Hide any event pins that are not supported by this Action's triggers in the advanced view
 	ForEachEventPinName([this, SupportedTriggerEvents](ETriggerEvent Event, FName PinName)
@@ -147,7 +148,7 @@ bool UK2Node_EnhancedInputAction::IsConnectionDisallowed(const UEdGraphPin* MyPi
 	if(MyPin->PinType.PinCategory == UEdGraphSchema_K2::PC_Exec && InputAction)
 	{
 		const ETriggerEvent Event = GetTriggerTypeFromExecPin(MyPin);
-		const ETriggerEventsSupported SupportedEvents = InputAction->GetSupportedTriggerEvents();
+		const ETriggerEventsSupported SupportedEvents = GetDefault<UBlueprintEditorSettings>()->bEnableInputTriggerSupportWarnings ? InputAction->GetSupportedTriggerEvents() : ETriggerEventsSupported::All;
 
 		if(!UInputTrigger::IsSupportedTriggerEvent(SupportedEvents, Event))
 		{
@@ -281,7 +282,7 @@ void UK2Node_EnhancedInputAction::ExpandNode(FKismetCompilerContext& CompilerCon
 		ETriggerEvent TriggerEvent;
 	};
 
-	const ETriggerEventsSupported SupportedTriggerEvents = InputAction->GetSupportedTriggerEvents();
+	const ETriggerEventsSupported SupportedTriggerEvents = GetDefault<UBlueprintEditorSettings>()->bEnableInputTriggerSupportWarnings ? InputAction->GetSupportedTriggerEvents() : ETriggerEventsSupported::All;
 	
 	TArray<ActivePinData> ActivePins;
 	ForEachEventPinName([this, &ActivePins, &SupportedTriggerEvents, &CompilerContext](ETriggerEvent Event, FName PinName) 
diff --git a/Engine/Source/Editor/BlueprintGraph/Private/BlueprintEditorSettings.cpp b/Engine/Source/Editor/BlueprintGraph/Private/BlueprintEditorSettings.cpp
--- a/Engine/Source/Editor/BlueprintGraph/Private/BlueprintEditorSettings.cpp
+++ b/Engine/Source/Editor/BlueprintGraph/Private/BlueprintEditorSettings.cpp
@@ -21,6 +21,7 @@ UBlueprintEditorSettings::UBlueprintEditorSettings(const FObjectInitializer& Obj
 	, bHideUnrelatedNodes(false)
 	, bShowShortTooltips(true)
 	// Workflow Settings
+	, bEnableInputTriggerSupportWarnings(false)
 	, bSplitContextTargetSettings(true)
 	, bExposeAllMemberComponentFunctions(true)
 	, bShowContextualFavorites(false)
diff --git a/Engine/Source/Editor/BlueprintGraph/Public/BlueprintEditorSettings.h b/Engine/Source/Editor/BlueprintGraph/Public/BlueprintEditorSettings.h
--- a/Engine/Source/Editor/BlueprintGraph/Public/BlueprintEditorSettings.h
+++ b/Engine/Source/Editor/BlueprintGraph/Public/BlueprintEditorSettings.h
@@ -88,6 +88,13 @@ public:
 
 // Workflow Settings
 public:
+	/** 
+	  * If enabled Input Action nodes will hide unsupported trigger pins and give warnings when using unsupported triggers.
+	  * This warning only works with triggers set up in an Input Action, not Input Mapping Contexts. 
+	  */
+	UPROPERTY(EditAnywhere, Config, Category=Workflow, meta=(DisplayName="Enhanced Input: Enable Input Trigger Support Warnings"))
+	bool bEnableInputTriggerSupportWarnings;
+	
 	/** If enabled, we'll save off your chosen target setting based off of the context (allowing you to have different preferences based off what you're operating on). */
 	UPROPERTY(EditAnywhere, config, Category=Workflow, meta=(DisplayName="Context Menu: Divide Context Target Preferences"))
 	bool bSplitContextTargetSettings;
