Fix a typo in the Input Action asset tooltip.

#rnx

diff --git a/Engine/Plugins/EnhancedInput/Source/InputEditor/Private/InputEditorModule.cpp b/Engine/Plugins/EnhancedInput/Source/InputEditor/Private/InputEditorModule.cpp
--- a/Engine/Plugins/EnhancedInput/Source/InputEditor/Private/InputEditorModule.cpp
+++ b/Engine/Plugins/EnhancedInput/Source/InputEditor/Private/InputEditorModule.cpp
@@ -138,7 +138,7 @@ public:
 	virtual FText GetName() const override { return NSLOCTEXT("AssetTypeActions", "AssetTypeActions_InputAction", "Input Action"); }
 	virtual uint32 GetCategories() override { return FInputEditorModule::GetInputAssetsCategory(); }
 	virtual FColor GetTypeColor() const override { return FColor(127, 255, 255); }
-	virtual FText GetAssetDescription(const FAssetData& AssetData) const override { return NSLOCTEXT("AssetTypeActions", "AssetTypeActions_InputActionDesc", "Represents an an abstract game action that can be mapped to arbitrary hardware input devices."); }
+	virtual FText GetAssetDescription(const FAssetData& AssetData) const override { return NSLOCTEXT("AssetTypeActions", "AssetTypeActions_InputActionDesc", "Represents an abstract game action that can be mapped to arbitrary hardware input devices."); }
 	virtual UClass* GetSupportedClass() const override { return UInputAction::StaticClass(); }
 	virtual bool HasActions(const TArray<UObject*>& InObjects) const override { return true; }
 	virtual void GetActions(const TArray<UObject*>& InObjects, struct FToolMenuSection& Section) override;
