Make 1D input actions use double pins instead of float pins like before. This makes them consistent with 2d and 3d actions, making math ops between types easier.

#rnx

diff --git a/Engine/Plugins/EnhancedInput/Source/InputBlueprintNodes/Private/K2Node_GetInputActionValue.cpp b/Engine/Plugins/EnhancedInput/Source/InputBlueprintNodes/Private/K2Node_GetInputActionValue.cpp
--- a/Engine/Plugins/EnhancedInput/Source/InputBlueprintNodes/Private/K2Node_GetInputActionValue.cpp
+++ b/Engine/Plugins/EnhancedInput/Source/InputBlueprintNodes/Private/K2Node_GetInputActionValue.cpp
@@ -50,7 +50,7 @@ const TMap<EInputActionValueType, FValueTypeData>& GetValueLookups()
 	static const TMap<EInputActionValueType, FValueTypeData> ValueLookups =
 	{
 		{ EInputActionValueType::Boolean, FValueTypeData(UEdGraphSchema_K2::PC_Boolean) },
-		{ EInputActionValueType::Axis1D, FValueTypeData(UEdGraphSchema_K2::PC_Real, UEdGraphSchema_K2::PC_Float) },
+		{ EInputActionValueType::Axis1D, FValueTypeData(UEdGraphSchema_K2::PC_Real, UEdGraphSchema_K2::PC_Double) },
 		{ EInputActionValueType::Axis2D, FValueTypeData(UEdGraphSchema_K2::PC_Struct, NAME_None, TBaseStructure<FVector2D>::Get()) },
 		{ EInputActionValueType::Axis3D, FValueTypeData(UEdGraphSchema_K2::PC_Struct, NAME_None, TBaseStructure<FVector>::Get()) },
 	};
