Undo making the "Get Input Action Value" node return a float. This made it so that 1D axis values didn't work in BP. The complete fix is to change the node expansion of the input action nodes to use doubles instead of floats

#rnx
#rnx

diff --git a/Engine/Plugins/EnhancedInput/Source/InputBlueprintNodes/Private/K2Node_GetInputActionValue.cpp b/Engine/Plugins/EnhancedInput/Source/InputBlueprintNodes/Private/K2Node_GetInputActionValue.cpp
--- a/Engine/Plugins/EnhancedInput/Source/InputBlueprintNodes/Private/K2Node_GetInputActionValue.cpp
+++ b/Engine/Plugins/EnhancedInput/Source/InputBlueprintNodes/Private/K2Node_GetInputActionValue.cpp
@@ -50,7 +50,7 @@ const TMap<EInputActionValueType, FValueTypeData>& GetValueLookups()
 	static const TMap<EInputActionValueType, FValueTypeData> ValueLookups =
 	{
 		{ EInputActionValueType::Boolean, FValueTypeData(UEdGraphSchema_K2::PC_Boolean) },
-		{ EInputActionValueType::Axis1D, FValueTypeData(UEdGraphSchema_K2::PC_Real, UEdGraphSchema_K2::PC_Double) },
+		{ EInputActionValueType::Axis1D, FValueTypeData(UEdGraphSchema_K2::PC_Real, UEdGraphSchema_K2::PC_Float) },
 		{ EInputActionValueType::Axis2D, FValueTypeData(UEdGraphSchema_K2::PC_Struct, NAME_None, TBaseStructure<FVector2D>::Get()) },
 		{ EInputActionValueType::Axis3D, FValueTypeData(UEdGraphSchema_K2::PC_Struct, NAME_None, TBaseStructure<FVector>::Get()) },
 	};
