Fixed a EI bug when using a float for the action value.

We were never setting a correct subcategory for the PC_Real pin, which led to subtle errors (eg: cast to string wouldn't work).

diff --git a/Engine/Plugins/EnhancedInput/Source/InputBlueprintNodes/Private/K2Node_EnhancedInputAction.cpp b/Engine/Plugins/EnhancedInput/Source/InputBlueprintNodes/Private/K2Node_EnhancedInputAction.cpp
--- a/Engine/Plugins/EnhancedInput/Source/InputBlueprintNodes/Private/K2Node_EnhancedInputAction.cpp
+++ b/Engine/Plugins/EnhancedInput/Source/InputBlueprintNodes/Private/K2Node_EnhancedInputAction.cpp
@@ -84,7 +84,7 @@ void UK2Node_EnhancedInputAction::AllocateDefaultPins()
 	
 	AdvancedPinDisplay = ENodeAdvancedPins::Hidden;
 
-	UEdGraphPin* ValuePin = CreatePin(EGPD_Output, UK2Node_GetInputActionValue::GetValueCategory(InputAction), UK2Node_GetInputActionValue::GetValueSubCategoryObject(InputAction), ActionValuePinName);
+	UEdGraphPin* ValuePin = CreatePin(EGPD_Output, UK2Node_GetInputActionValue::GetValueCategory(InputAction), UK2Node_GetInputActionValue::GetValueSubCategory(InputAction), UK2Node_GetInputActionValue::GetValueSubCategoryObject(InputAction), ActionValuePinName);
 	
 	Schema->SetPinAutogeneratedDefaultValueBasedOnType(ValuePin);
 
