Add the Action Type as an AssetRegisterySearchable property so that it shows up when you hover over an Input Action asset in the content browser. This makes it easier to see what type of input action it is!

#rnx

diff --git a/Engine/Plugins/EnhancedInput/Source/EnhancedInput/Public/InputAction.h b/Engine/Plugins/EnhancedInput/Source/EnhancedInput/Public/InputAction.h
--- a/Engine/Plugins/EnhancedInput/Source/EnhancedInput/Public/InputAction.h
+++ b/Engine/Plugins/EnhancedInput/Source/EnhancedInput/Public/InputAction.h
@@ -59,7 +59,7 @@ public:
 	bool bReserveAllMappings = false;	// TODO: Need something more complex than this?
 
 	// The type that this action returns from a GetActionValue query or action event
-	UPROPERTY(EditAnywhere, BlueprintReadOnly, Category = Action)
+	UPROPERTY(EditAnywhere, BlueprintReadOnly, Category = Action, AssetRegistrySearchable)
 	EInputActionValueType ValueType = EInputActionValueType::Boolean;
 	
 	/**
