Fixing a tooltip typo in "Triggers" mouse-over text

#rnx

diff --git a/Engine/Plugins/EnhancedInput/Source/EnhancedInput/Public/InputAction.h b/Engine/Plugins/EnhancedInput/Source/EnhancedInput/Public/InputAction.h
--- a/Engine/Plugins/EnhancedInput/Source/EnhancedInput/Public/InputAction.h
+++ b/Engine/Plugins/EnhancedInput/Source/EnhancedInput/Public/InputAction.h
@@ -62,6 +62,14 @@ public:
 	UPROPERTY(EditAnywhere, BlueprintReadOnly, Category = Action, AssetRegistrySearchable)
 	EInputActionValueType ValueType = EInputActionValueType::Boolean;
 	
+	/**
+	* Trigger qualifiers. If any trigger qualifiers exist the action will not trigger unless:
+	* At least one Explicit trigger in this list has been met.
+	* All Implicit triggers in this list are met.
+	*/
+	UPROPERTY(EditAnywhere, Instanced, BlueprintReadWrite, Category = Action)
+	TArray<TObjectPtr<UInputTrigger>> Triggers;
+
 	/**
 	* Modifiers are applied to the final action value.
 	* These are applied sequentially in array order.
