InputAction: Fixing a compile error

diff --git a/Engine/Plugins/EnhancedInput/Source/EnhancedInput/Public/InputAction.h b/Engine/Plugins/EnhancedInput/Source/EnhancedInput/Public/InputAction.h
--- a/Engine/Plugins/EnhancedInput/Source/EnhancedInput/Public/InputAction.h
+++ b/Engine/Plugins/EnhancedInput/Source/EnhancedInput/Public/InputAction.h
@@ -33,14 +33,6 @@ public:
 	 */
 	ETriggerEventsSupported GetSupportedTriggerEvents() const;
 #endif
-
-	/**
-	* Trigger qualifiers. If any trigger qualifiers exist the action will not trigger unless:
-	* At least one Explicit trigger in this list is be met.
-	* All Implicit triggers in this list are met.
-	*/
-	UPROPERTY(EditAnywhere, Instanced, BlueprintReadWrite, Category = Action)
-	TArray<TObjectPtr<UInputTrigger>> Triggers;
 	
 	// A localized descriptor of this input action
 	UPROPERTY(EditAnywhere, BlueprintReadOnly, Category = "Description")
