Making it clear that a key is not necessary/usable for Combo actions

diff --git a/Engine/Plugins/EnhancedInput/Source/InputEditor/Private/InputCustomizations.cpp b/Engine/Plugins/EnhancedInput/Source/InputEditor/Private/InputCustomizations.cpp
--- a/Engine/Plugins/EnhancedInput/Source/InputEditor/Private/InputCustomizations.cpp
+++ b/Engine/Plugins/EnhancedInput/Source/InputEditor/Private/InputCustomizations.cpp
@@ -8,7 +8,6 @@
 #include "EnhancedActionKeyMapping.h"
 #include "IDetailChildrenBuilder.h"
 #include "InputMappingContext.h"
-#include "KeyStructCustomization.h"
 #include "PropertyCustomizationHelpers.h"
 #include "EnhancedInputDeveloperSettings.h"
 #include "InputEditorModule.h"
@@ -46,6 +45,9 @@ void FEnhancedActionMappingCustomization::CustomizeHeader(TSharedRef<IPropertyHa
 
 	// Grab the FKey property
 	TSharedPtr<IPropertyHandle> KeyHandle = PropertyHandle->GetChildHandle(GET_MEMBER_NAME_CHECKED(FEnhancedActionKeyMapping, Key));
+	TSharedPtr<IPropertyHandle> TriggersHandle = PropertyHandle->GetChildHandle(GET_MEMBER_NAME_CHECKED(FEnhancedActionKeyMapping, Triggers));
+
+	TriggersHandle->SetOnPropertyValueChanged(FSimpleDelegate::CreateSP(this, &FEnhancedActionMappingCustomization::OnTriggersChanged));
 
 	TSharedRef<SWidget> RemoveButton = PropertyCustomizationHelpers::MakeDeleteButton(FSimpleDelegate::CreateSP(this, &FEnhancedActionMappingCustomization::RemoveMappingButton_OnClick),
 		LOCTEXT("RemoveMappingToolTip", "Remove Mapping"));
@@ -56,7 +58,11 @@ void FEnhancedActionMappingCustomization::CustomizeHeader(TSharedRef<IPropertyHa
 	// TODO: Use FDetailArrayBuilder?
 
 	// Pass our header row into the key struct customizeheader method so it populates our row with the key struct header
-	StaticCastSharedPtr<FKeyStructCustomization>(KeyStructInstance)->CustomizeHeaderOnlyWithButton(KeyHandle.ToSharedRef(), HeaderRow, CustomizationUtils, RemoveButton);
+	KeyStructCustomization = StaticCastSharedPtr<FKeyStructCustomization>(KeyStructInstance);
+	const bool bContainsComboTrigger = DoesTriggerArrayContainCombo();
+	KeyStructCustomization->bDisplayIcon = bContainsComboTrigger;
+	KeyStructCustomization->bEnableKeySelector = !bContainsComboTrigger;
+	KeyStructCustomization->CustomizeHeaderOnlyWithButton(KeyHandle.ToSharedRef(), HeaderRow, CustomizationUtils, RemoveButton);
 }
 
 void FEnhancedActionMappingCustomization::CustomizeChildren(TSharedRef<IPropertyHandle> PropertyHandle, IDetailChildrenBuilder& ChildBuilder, IPropertyTypeCustomizationUtils& CustomizationUtils)
@@ -84,6 +90,46 @@ void FEnhancedActionMappingCustomization::RemoveMappingButton_OnClick() const
 	}
 }
 
+void FEnhancedActionMappingCustomization::OnTriggersChanged() const
+{
+	const bool bContainsComboTrigger = DoesTriggerArrayContainCombo();
+	KeyStructCustomization->bDisplayIcon = bContainsComboTrigger;
+	KeyStructCustomization->bEnableKeySelector = !bContainsComboTrigger;
+}
+
+bool FEnhancedActionMappingCustomization::DoesTriggerArrayContainCombo() const
+{
+	if (MappingPropertyHandle)
+	{
+		if (TSharedPtr<IPropertyHandle> TriggersHandle = MappingPropertyHandle->GetChildHandle(GET_MEMBER_NAME_CHECKED(FEnhancedActionKeyMapping, Triggers)))
+		{
+			// getting data for the trigger array
+			void* Data = nullptr;
+			TriggersHandle->GetValueData(Data);
+			if (Data)
+			{
+				FProperty* TriggersProperty = TriggersHandle->GetProperty();
+				FArrayProperty* TriggersArrayProperty = CastField<FArrayProperty>(TriggersProperty);
+				FScriptArrayHelper ArrayHelper(TriggersArrayProperty, Data);
+	
+				for (int32 i = 0; i < ArrayHelper.Num(); i++)
+				{
+					// Make sure we can cast this to a input trigger and if it's a combo we can return true
+					if (UInputTrigger** InputComboTrigger = reinterpret_cast<UInputTrigger**>(ArrayHelper.GetRawPtr(i)))
+					{
+						if ((*InputComboTrigger) && (*InputComboTrigger)->IsA(UInputTriggerComboAction::StaticClass()))
+						{
+							return true;
+						}
+					}
+				}
+			}
+		}
+	}
+	
+	return false;
+}
+
 //////////////////////////////////////////////////////////
 // FEnhancedInputDeveloperSettingsCustomization
 
diff --git a/Engine/Plugins/EnhancedInput/Source/InputEditor/Private/InputCustomizations.h b/Engine/Plugins/EnhancedInput/Source/InputEditor/Private/InputCustomizations.h
--- a/Engine/Plugins/EnhancedInput/Source/InputEditor/Private/InputCustomizations.h
+++ b/Engine/Plugins/EnhancedInput/Source/InputEditor/Private/InputCustomizations.h
@@ -6,6 +6,7 @@
 #include "IDetailCustomization.h"
 #include "IDetailCustomNodeBuilder.h"
 #include "IPropertyTypeCustomization.h"
+#include "KeyStructCustomization.h"
 
 class IDetailLayoutBuilder;
 
@@ -31,12 +32,16 @@ public:
 	virtual void CustomizeHeader(TSharedRef<IPropertyHandle> PropertyHandle, FDetailWidgetRow& HeaderRow, IPropertyTypeCustomizationUtils& CustomizationUtils) override;
 	virtual void CustomizeChildren(TSharedRef<IPropertyHandle> PropertyHandle, IDetailChildrenBuilder& ChildBuilder, IPropertyTypeCustomizationUtils& CustomizationUtils) override;
 
+	/** Searches the trigger array of an input action - returns true if any combo input triggers are present */
+	bool DoesTriggerArrayContainCombo() const;
 private:
 
 	void RemoveMappingButton_OnClick() const;
+	void OnTriggersChanged() const;
 
 	TSharedPtr<IPropertyTypeCustomization> KeyStructInstance;
 	TSharedPtr<IPropertyHandle> MappingPropertyHandle;
+	TSharedPtr<FKeyStructCustomization> KeyStructCustomization;
 };
 
 /**
diff --git a/Engine/Source/Editor/DetailCustomizations/Private/KeyStructCustomization.cpp b/Engine/Source/Editor/DetailCustomizations/Private/KeyStructCustomization.cpp
--- a/Engine/Source/Editor/DetailCustomizations/Private/KeyStructCustomization.cpp
+++ b/Engine/Source/Editor/DetailCustomizations/Private/KeyStructCustomization.cpp
@@ -77,6 +77,10 @@ void FKeyStructCustomization::CustomizeHeaderOnlyWithButton(TSharedRef<class IPr
 			.Font(StructCustomizationUtils.GetRegularFont())
 			.AllowClear(!StructPropertyHandle->GetProperty()->HasAnyPropertyFlags(CPF_NoClear))
 		    .FilterBlueprintBindable(false)
+		    .IsEnabled_Lambda([this]() -> bool
+		    {
+		    	return bEnableKeySelector;
+		    })
 		]
 		+ SHorizontalBox::Slot()
 		.Padding(InputSettingsDetails::InputConstants::PropertyPadding)
@@ -86,11 +90,30 @@ void FKeyStructCustomization::CustomizeHeaderOnlyWithButton(TSharedRef<class IPr
 		[
 			Button
 		]
+		+ SHorizontalBox::Slot()
+		.Padding(InputSettingsDetails::InputConstants::PropertyPadding)
+		.HAlign(HAlign_Center)
+		.VAlign(VAlign_Center)
+		.AutoWidth()
+		[
+			SNew(SImage)
+			.Visibility_Lambda([this]() {
+				return bDisplayIcon ? EVisibility::Visible : EVisibility::Hidden;
+			})
+			.Image(FAppStyle::GetBrush("Icons.Info"))
+			.ToolTipText(LOCTEXT("ComboTriggerKeyWarningText", "This Key is NOT part of the Combo. This action will be triggered when the actions in the Combo Actions array are completed."))
+			.ColorAndOpacity(FSlateColor::UseForeground())
+        ]
 	];
 }
 
 TOptional<FKey> FKeyStructCustomization::GetCurrentKey() const
 {
+	if (!bEnableKeySelector)
+	{
+		PropertyHandle->SetValueFromFormattedString(TEXT("None"));
+	}
+	
 	TArray<void*> StructPtrs;
 	PropertyHandle->AccessRawData(StructPtrs);
 
diff --git a/Engine/Source/Editor/DetailCustomizations/Public/KeyStructCustomization.h b/Engine/Source/Editor/DetailCustomizations/Public/KeyStructCustomization.h
--- a/Engine/Source/Editor/DetailCustomizations/Public/KeyStructCustomization.h
+++ b/Engine/Source/Editor/DetailCustomizations/Public/KeyStructCustomization.h
@@ -30,7 +30,8 @@ public:
 	// Helper variant that generates the key struct in the header and appends a single button at the end
 	// TODO: Is there a better way?
 	void CustomizeHeaderOnlyWithButton(TSharedRef<class IPropertyHandle> StructPropertyHandle, class FDetailWidgetRow& HeaderRow, IPropertyTypeCustomizationUtils& StructCustomizationUtils, TSharedRef<SWidget> Button);
-
+	bool bDisplayIcon = false;
+	bool bEnableKeySelector = true;
 public:
 
 	/**
