Add the bShouldOnlyTriggerLastActionInChord option to Enhanced Input, defaulted to on. This will track input actions that have Chorded Triggers on them, and allow the filtering of triggered delegates that are fired based on which chorded actions are being triggered this tick.

This will allow you to have a chorded action trigger separately while ebing bound to the same key as another, non-chorded action. For example, you can have Left Trigger fire on pressed, right trigger fire on pressed, and have a Chorded action for when they are triggered together.

The example above used to result in the Left trigger being fired at the same time as the chorded trigger, making it difficult to make intuitive chords. Now, only the chorded action will be fired.

diff --git a/Engine/Plugins/EnhancedInput/Source/EnhancedInput/Private/EnhancedInputSubsystemInterface.cpp b/Engine/Plugins/EnhancedInput/Source/EnhancedInput/Private/EnhancedInputSubsystemInterface.cpp
--- a/Engine/Plugins/EnhancedInput/Source/EnhancedInput/Private/EnhancedInputSubsystemInterface.cpp
+++ b/Engine/Plugins/EnhancedInput/Source/EnhancedInput/Private/EnhancedInputSubsystemInterface.cpp
@@ -483,21 +483,31 @@ void DeepCopyPtrArray(const TArray<T*>& From, TArray<T*>& To)
 }
 
 // TODO: This should be a delegate (along with InjectChordBlockers), moving chording out of the underlying subsystem and enabling implementation of custom mapping handlers.
-TArray<FEnhancedActionKeyMapping> ReorderMappings(const TArray<FEnhancedActionKeyMapping>& UnorderedMappings)
+/**
+ * Reorder the given UnordedMappings such that chording mappings > chorded mappings > everything else.
+ * This is used to ensure mappings within a single context are evaluated in the correct order to support chording.
+ * Populate the DependentChordActions array with any chorded triggers so that we can detect which ones should be triggered
+ * later. 
+ */
+TArray<FEnhancedActionKeyMapping> IEnhancedInputSubsystemInterface::ReorderMappings(const TArray<FEnhancedActionKeyMapping>& UnorderedMappings, TArray<UEnhancedPlayerInput::FDependentChordTracker>& OUT DependentChordActions)
 {
-	// Reorder mappings such that chording mappings > chorded mappings > everything else. This is used to ensure mappings within a single context are evaluated in the correct order to support chording.
-
 	TSet<const UInputAction*> ChordingActions;
 
 	// Gather all chording actions within a mapping's triggers.
-	auto GatherChordingActions = [&ChordingActions](const FEnhancedActionKeyMapping& Mapping) {
+	auto GatherChordingActions = [&ChordingActions, &DependentChordActions](const FEnhancedActionKeyMapping& Mapping)
+	{
 		bool bFoundChordTrigger = false;
-		auto EvaluateTriggers = [&Mapping, &ChordingActions, &bFoundChordTrigger](const TArray<UInputTrigger*>& Triggers) {
+		auto EvaluateTriggers = [&Mapping, &ChordingActions, &bFoundChordTrigger, &DependentChordActions](const TArray<UInputTrigger*>& Triggers)
+		{
 			for (const UInputTrigger* Trigger : Triggers)
 			{
 				if (const UInputTriggerChordAction* ChordTrigger = Cast<const UInputTriggerChordAction>(Trigger))
 				{
 					ChordingActions.Add(ChordTrigger->ChordAction);
+
+					// Keep track of the action itself, and the action it is dependant on
+					DependentChordActions.Emplace(UEnhancedPlayerInput::FDependentChordTracker { Mapping.Action, ChordTrigger->ChordAction });
+					
 					bFoundChordTrigger = true;
 				}
 			}
@@ -623,13 +633,16 @@ void IEnhancedInputSubsystemInterface::RebuildControlMappings()
 
 	TArray<int32> ChordedMappings;
 
+	// Reset the tracking of dependant chord actions on the player input
+	PlayerInput->DependentChordActions.Reset();
+
 	for (const TPair<TObjectPtr<const UInputMappingContext>, int32>& ContextPair : OrderedInputContexts)
 	{
 		// Don't apply context specific keys immediately, allowing multiple mappings to the same key within the same context if required.
 		TArray<FKey> ContextAppliedKeys;
 
 		const UInputMappingContext* MappingContext = ContextPair.Key;
-		TArray<FEnhancedActionKeyMapping> OrderedMappings = ReorderMappings(MappingContext->GetMappings());
+		TArray<FEnhancedActionKeyMapping> OrderedMappings = ReorderMappings(MappingContext->GetMappings(), PlayerInput->DependentChordActions);
 
 		for (FEnhancedActionKeyMapping& Mapping : OrderedMappings)
 		{
diff --git a/Engine/Plugins/EnhancedInput/Source/EnhancedInput/Private/EnhancedPlayerInput.cpp b/Engine/Plugins/EnhancedInput/Source/EnhancedInput/Private/EnhancedPlayerInput.cpp
--- a/Engine/Plugins/EnhancedInput/Source/EnhancedInput/Private/EnhancedPlayerInput.cpp
+++ b/Engine/Plugins/EnhancedInput/Source/EnhancedInput/Private/EnhancedPlayerInput.cpp
@@ -8,6 +8,18 @@
 #include "InputTriggers.h"
 #include "Engine/World.h"
 #include "GameFramework/WorldSettings.h"
+#include "EnhancedInputDeveloperSettings.h"
+
+namespace UE
+{
+	namespace Input
+	{
+		static int32 ShouldOnlyTriggerLastActionInChord = 1;
+		static FAutoConsoleVariableRef CVarShouldOnlyTriggerLastActionInChord(TEXT("EnhancedInput.OnlyTriggerLastActionInChord"),
+			ShouldOnlyTriggerLastActionInChord,
+			TEXT("Should only the last action in a ChordedAction trigger be fired? If this is disabled, then the dependant chords will be fired as well"));
+	}
+}
 
 // NOTE: Enum order represents firing priority(lowest to highest) and is important as multiple keys bound to the same action may generate differing trigger event states.
 enum class ETriggerEventInternal : uint8
@@ -442,7 +454,13 @@ void UEnhancedPlayerInput::ProcessInputStack(const TArray<UInputComponent*>& Inp
 					else
 					{
 						TriggeredDelegates.Emplace(Binding->Clone());
-					}	
+					}
+
+					// Keep track of the triggered actions this tick so that we can quickly look them up later when determining chorded action state
+					if (BoundTriggerEvent == ETriggerEvent::Triggered)
+					{
+						TriggeredActionsThisTick.Add(ActionData->GetSourceAction());
+					}
 				}
 			}
 		}
@@ -450,13 +468,38 @@ void UEnhancedPlayerInput::ProcessInputStack(const TArray<UInputComponent*>& Inp
 		// Action all delegates that triggered this tick, in the order in which they triggered.
 		for (TUniquePtr<FEnhancedInputActionEventBinding>& Delegate : TriggeredDelegates)
 		{
-			// Search for the action instance data a second time as a previous delegate call may have deleted it.
-			if (const FInputActionInstance* ActionData = FindActionInstanceData(Delegate->GetAction()))
+			const UInputAction* DelegateAction = Delegate->GetAction();
+			bool bCanTrigger = true;
+
+			if (UE::Input::ShouldOnlyTriggerLastActionInChord)
 			{
-				Delegate->Execute(*ActionData);
+				// If this delegate is referenced by a UInputTriggerChordAction::ChordAction
+				// then we only want to trigger it the referencing action is not triggered
+				for (const UEnhancedPlayerInput::FDependentChordTracker& DepAction : DependentChordActions)
+				{
+					if (DepAction.DependantAction && DepAction.DependantAction == DelegateAction)
+					{
+						bCanTrigger = !TriggeredActionsThisTick.Contains(DepAction.SourceAction);
+						if(!bCanTrigger)
+						{
+							UE_LOG(LogTemp, Warning, TEXT("'%s' action was cancelled, its dependant on '%s'"), *DelegateAction->GetName(), *DepAction.SourceAction->GetName());
+						}
+						break;
+					}
+				}
+			}
+
+			if (bCanTrigger)
+			{
+				// Search for the action instance data a second time as a previous delegate call may have deleted it.
+				if (const FInputActionInstance* ActionData = FindActionInstanceData(DelegateAction))
+				{
+					Delegate->Execute(*ActionData);
+				}
 			}
 		}
 		TriggeredDelegates.Reset();
+		TriggeredActionsThisTick.Reset();
 
 		// Update action value bindings
 		for (const FEnhancedInputActionValueBinding& Binding : IC->GetActionValueBindings())
diff --git a/Engine/Plugins/EnhancedInput/Source/EnhancedInput/Public/EnhancedInputDeveloperSettings.h b/Engine/Plugins/EnhancedInput/Source/EnhancedInput/Public/EnhancedInputDeveloperSettings.h
--- a/Engine/Plugins/EnhancedInput/Source/EnhancedInput/Public/EnhancedInputDeveloperSettings.h
+++ b/Engine/Plugins/EnhancedInput/Source/EnhancedInput/Public/EnhancedInputDeveloperSettings.h
@@ -22,4 +22,13 @@ public:
 	 */
 	UPROPERTY(EditAnywhere, Category = "Enhanced Input")
 	FPerPlatformSettings PlatformSettings;
+
+	/**
+	 * If true, then only the last action in a ChordedAction trigger will be fired.
+	 * This means that only the action that has the ChordedTrigger on it will be fired, not the individual steps.
+	 * 
+	 * Default value is true.
+	 */
+	UPROPERTY(EditAnywhere, Category = "Enhanced Input", meta=(ConsoleVariable="EnhancedInput.OnlyTriggerLastActionInChord"))
+	bool bShouldOnlyTriggerLastActionInChord = true;
 };
\ No newline at end of file
diff --git a/Engine/Plugins/EnhancedInput/Source/EnhancedInput/Public/EnhancedInputSubsystemInterface.h b/Engine/Plugins/EnhancedInput/Source/EnhancedInput/Public/EnhancedInputSubsystemInterface.h
--- a/Engine/Plugins/EnhancedInput/Source/EnhancedInput/Public/EnhancedInputSubsystemInterface.h
+++ b/Engine/Plugins/EnhancedInput/Source/EnhancedInput/Public/EnhancedInputSubsystemInterface.h
@@ -7,6 +7,7 @@
 #include "InputMappingQuery.h"
 #include "UObject/Interface.h"
 #include "EnhancedActionKeyMapping.h"
+#include "EnhancedPlayerInput.h"
 
 #include "EnhancedInputSubsystemInterface.generated.h"
 
@@ -251,8 +252,16 @@ private:
 	void RemoveForcedInput(FKey Key);
 
 	void InjectChordBlockers(const TArray<int32>& ChordedMappings);
-	bool HasTriggerWith(TFunctionRef<bool(const class UInputTrigger*)> TestFn, const TArray<class UInputTrigger*>& Triggers);
+	bool HasTriggerWith(TFunctionRef<bool(const class UInputTrigger*)> TestFn, const TArray<class UInputTrigger*>& Triggers);	
 
+	/**
+	 * Reorder the given UnordedMappings such that chording mappings > chorded mappings > everything else.
+	 * This is used to ensure mappings within a single context are evaluated in the correct order to support chording.
+	 * Populate the DependentChordActions array with any chorded triggers so that we can detect which ones should be triggered
+	 * later. 
+	 */
+	TArray<FEnhancedActionKeyMapping> ReorderMappings(const TArray<FEnhancedActionKeyMapping>& UnorderedMappings, TArray<UEnhancedPlayerInput::FDependentChordTracker>& OUT DependentChordActions);
+	
 	/**
 	 * Reapply all control mappings to players pending a rebuild
 	 */
@@ -267,6 +276,12 @@ private:
 	/** A map of any player mapped keys to the key that they should redirect to instead */
 	TMap<FName, FKey> PlayerMappedSettings;
 
+	/**
+	 * A map of input actions with a Chorded trigger, mapped to the action they are dependent on.
+	 * The Key is the Input Action with the Chorded Trigger, and the value is the action it is dependant on.
+	 */
+	TMap<TObjectPtr<const UInputAction>, TObjectPtr<const UInputAction>> ChordedActionDependencies;
+
 	/**
 	 * A map of the currently applied mapping context redirects. This is populated in RebuildControlMappings
 	 * with any InputMappingContexts that have been redirected on this local player from the platform settings.
diff --git a/Engine/Plugins/EnhancedInput/Source/EnhancedInput/Public/EnhancedPlayerInput.h b/Engine/Plugins/EnhancedInput/Source/EnhancedInput/Public/EnhancedPlayerInput.h
--- a/Engine/Plugins/EnhancedInput/Source/EnhancedInput/Public/EnhancedPlayerInput.h
+++ b/Engine/Plugins/EnhancedInput/Source/EnhancedInput/Public/EnhancedPlayerInput.h
@@ -106,6 +106,9 @@ private:
 	/** Actions which had actuated events at the last call to ProcessInputStack (held/pressed/released) */
 	TSet<const UInputAction*> ActionsWithEventsThisTick;
 
+	/** Actions that have been triggered this tick and have a delegate that may be fired */
+	TSet<TObjectPtr<const UInputAction>> TriggeredActionsThisTick;
+
 	/**
 	 * A map of Keys to the amount they were depressed this frame. This is reset with each call to ProcessInputStack
 	 * and is populated within UEnhancedPlayerInput::InputKey.
@@ -130,6 +133,22 @@ private:
 	/** Last frame's injected inputs */
 	TSet<const UInputAction*> LastInjectedActions;
 
+	/** Used to keep track of Input Actions that have UInputTriggerChordAction triggers on them */
+	struct FDependentChordTracker
+	{
+		/** The Input Action that has the UInputTriggerChordAction on it */
+		TObjectPtr<const UInputAction> SourceAction;
+		
+		/** The action that is referenced by the SourceAction's Chord trigger */
+		TObjectPtr<const UInputAction> DependantAction;
+	};
+	
+	/**
+	 * Array of all dependant Input Action's with Chord triggers on them.
+	 * Populated by IEnhancedInputSubsystemInterface::ReorderMappings
+	 */
+	TArray<FDependentChordTracker> DependentChordActions;
+
 	/** The last time of the last frame that was processed in ProcessPlayerInput */
 	float LastFrameTime = 0.0f;
 
