Make EnhancedInput action event nodes double precision floats instead of single precision.
This is the same change as 21593878, but with the correct BP auto conversion function added. Without the conversion function, the node would fail to expand because there was not a FInputActionValue -> double function.

#rnx

diff --git a/Engine/Plugins/EnhancedInput/Source/EnhancedInput/Private/EnhancedInputLibrary.cpp b/Engine/Plugins/EnhancedInput/Source/EnhancedInput/Private/EnhancedInputLibrary.cpp
--- a/Engine/Plugins/EnhancedInput/Source/EnhancedInput/Private/EnhancedInputLibrary.cpp
+++ b/Engine/Plugins/EnhancedInput/Source/EnhancedInput/Private/EnhancedInputLibrary.cpp
@@ -46,7 +46,7 @@ FInputActionValue UEnhancedInputLibrary::GetBoundActionValue(AActor* Actor, cons
 }
 
 
-void UEnhancedInputLibrary::BreakInputActionValue(FInputActionValue InActionValue, float& X, float& Y, float& Z)
+void UEnhancedInputLibrary::BreakInputActionValue(FInputActionValue InActionValue, double& X, double& Y, double& Z)
 {
 	FVector AsAxis3D = InActionValue.Get<FInputActionValue::Axis3D>();
 	X = AsAxis3D.X;
@@ -54,7 +54,7 @@ void UEnhancedInputLibrary::BreakInputActionValue(FInputActionValue InActionValu
 	Z = AsAxis3D.Z;
 }
 
-FInputActionValue UEnhancedInputLibrary::MakeInputActionValue(float X, float Y, float Z, const FInputActionValue& MatchValueType)
+FInputActionValue UEnhancedInputLibrary::MakeInputActionValue(double X, double Y, double Z, const FInputActionValue& MatchValueType)
 {
 	return FInputActionValue(MatchValueType.GetValueType(), FVector(X, Y, Z));
 }
@@ -66,9 +66,9 @@ bool UEnhancedInputLibrary::Conv_InputActionValueToBool(FInputActionValue InValu
 	return InValue.Get<bool>();
 }
 
-float UEnhancedInputLibrary::Conv_InputActionValueToAxis1D(FInputActionValue InValue)
+double UEnhancedInputLibrary::Conv_InputActionValueToAxis1D(FInputActionValue InValue)
 {
-	return InValue.Get<FInputActionValue::Axis1D>();
+	return static_cast<double>(InValue.Get<FInputActionValue::Axis1D>());
 }
 
 FVector2D UEnhancedInputLibrary::Conv_InputActionValueToAxis2D(FInputActionValue InValue)
diff --git a/Engine/Plugins/EnhancedInput/Source/EnhancedInput/Public/EnhancedInputLibrary.h b/Engine/Plugins/EnhancedInput/Source/EnhancedInput/Public/EnhancedInputLibrary.h
--- a/Engine/Plugins/EnhancedInput/Source/EnhancedInput/Public/EnhancedInputLibrary.h
+++ b/Engine/Plugins/EnhancedInput/Source/EnhancedInput/Public/EnhancedInputLibrary.h
@@ -38,14 +38,14 @@ public:
 
 	/** Breaks an ActionValue into X, Y, Z. Axes not supported by value type will be 0. */
 	UFUNCTION(BlueprintPure, Category = "Input", meta = (NativeBreakFunc))
-	static void BreakInputActionValue(FInputActionValue InActionValue, float& X, float& Y, float& Z);
+	static void BreakInputActionValue(FInputActionValue InActionValue, double& X, double& Y, double& Z);
 
 	/**
 	 * Builds an ActionValue from X, Y, Z. Inherits type from an existing ActionValue. Ignores axis values unused by the provided value type.
 	 * @note Intended for use in Input Modifier Modify Raw overloads to modify an existing Input Action Value.
 	 */
 	UFUNCTION(BlueprintPure, Category = "Input", meta = (Keywords = "construct build", NativeMakeFunc))
-	static FInputActionValue MakeInputActionValue(float X, float Y, float Z, const FInputActionValue& MatchValueType);
+	static FInputActionValue MakeInputActionValue(double X, double Y, double Z, const FInputActionValue& MatchValueType);
 
 
 	// Internal helper functionality
@@ -60,9 +60,9 @@ public:
 	UFUNCTION(BlueprintPure, BlueprintInternalUseOnly, meta = (BlueprintAutocast))
 	static bool Conv_InputActionValueToBool(FInputActionValue InValue);
 
-	/** Interpret an InputActionValue as a 1D axis (float) input */
+	/** Interpret an InputActionValue as a 1D axis (double) input */
 	UFUNCTION(BlueprintPure, BlueprintInternalUseOnly, meta = (BlueprintAutocast))
-	static float Conv_InputActionValueToAxis1D(FInputActionValue InValue);
+	static double Conv_InputActionValueToAxis1D(FInputActionValue InValue);
 
 	/** Interpret an InputActionValue as a 2D axis (Vector2D) input */
 	UFUNCTION(BlueprintPure, BlueprintInternalUseOnly, meta = (BlueprintAutocast))
diff --git a/Engine/Plugins/EnhancedInput/Source/InputBlueprintNodes/Private/K2Node_GetInputActionValue.cpp b/Engine/Plugins/EnhancedInput/Source/InputBlueprintNodes/Private/K2Node_GetInputActionValue.cpp
--- a/Engine/Plugins/EnhancedInput/Source/InputBlueprintNodes/Private/K2Node_GetInputActionValue.cpp
+++ b/Engine/Plugins/EnhancedInput/Source/InputBlueprintNodes/Private/K2Node_GetInputActionValue.cpp
@@ -50,7 +50,7 @@ const TMap<EInputActionValueType, FValueTypeData>& GetValueLookups()
 	static const TMap<EInputActionValueType, FValueTypeData> ValueLookups =
 	{
 		{ EInputActionValueType::Boolean, FValueTypeData(UEdGraphSchema_K2::PC_Boolean) },
-		{ EInputActionValueType::Axis1D, FValueTypeData(UEdGraphSchema_K2::PC_Real, UEdGraphSchema_K2::PC_Float) },
+		{ EInputActionValueType::Axis1D, FValueTypeData(UEdGraphSchema_K2::PC_Real, UEdGraphSchema_K2::PC_Double) },
 		{ EInputActionValueType::Axis2D, FValueTypeData(UEdGraphSchema_K2::PC_Struct, NAME_None, TBaseStructure<FVector2D>::Get()) },
 		{ EInputActionValueType::Axis3D, FValueTypeData(UEdGraphSchema_K2::PC_Struct, NAME_None, TBaseStructure<FVector>::Get()) },
 	};
