Unshelved from pending changelist '21830153':

Change the Make/Break nodes for the FInputActionValue struct in blueprints to correctly take in a EInputActionValueType enum instead of the incorrect FInputActionValue struct. Before, you could infinitely split the pin struct because it was referencing the same struct type that it was trying to make.

diff --git a/Engine/Plugins/EnhancedInput/Source/EnhancedInput/Private/EnhancedInputLibrary.cpp b/Engine/Plugins/EnhancedInput/Source/EnhancedInput/Private/EnhancedInputLibrary.cpp
--- a/Engine/Plugins/EnhancedInput/Source/EnhancedInput/Private/EnhancedInputLibrary.cpp
+++ b/Engine/Plugins/EnhancedInput/Source/EnhancedInput/Private/EnhancedInputLibrary.cpp
@@ -46,18 +46,26 @@ FInputActionValue UEnhancedInputLibrary::GetBoundActionValue(AActor* Actor, cons
 }
 
 
-void UEnhancedInputLibrary::BreakInputActionValue(FInputActionValue InActionValue, double& X, double& Y, double& Z)
+void UEnhancedInputLibrary::BreakInputActionValue(FInputActionValue InActionValue, double& X, double& Y, double& Z, EInputActionValueType& Type)
 {
 	FVector AsAxis3D = InActionValue.Get<FInputActionValue::Axis3D>();
 	X = AsAxis3D.X;
 	Y = AsAxis3D.Y;
 	Z = AsAxis3D.Z;
+	Type = InActionValue.GetValueType();
 }
 
+PRAGMA_DISABLE_DEPRECATION_WARNINGS
 FInputActionValue UEnhancedInputLibrary::MakeInputActionValue(double X, double Y, double Z, const FInputActionValue& MatchValueType)
 {
 	return FInputActionValue(MatchValueType.GetValueType(), FVector(X, Y, Z));
 }
+PRAGMA_ENABLE_DEPRECATION_WARNINGS
+
+FInputActionValue UEnhancedInputLibrary::MakeInputActionValueOfType(double X, double Y, double Z, const EInputActionValueType ValueType)
+{
+	return FInputActionValue(ValueType, FVector(X, Y, Z));
+}
 
 // FInputActionValue type conversions
 
diff --git a/Engine/Plugins/EnhancedInput/Source/EnhancedInput/Public/EnhancedInputLibrary.h b/Engine/Plugins/EnhancedInput/Source/EnhancedInput/Public/EnhancedInputLibrary.h
--- a/Engine/Plugins/EnhancedInput/Source/EnhancedInput/Public/EnhancedInputLibrary.h
+++ b/Engine/Plugins/EnhancedInput/Source/EnhancedInput/Public/EnhancedInputLibrary.h
@@ -38,15 +38,18 @@ public:
 
 	/** Breaks an ActionValue into X, Y, Z. Axes not supported by value type will be 0. */
 	UFUNCTION(BlueprintPure, Category = "Input", meta = (NativeBreakFunc))
-	static void BreakInputActionValue(FInputActionValue InActionValue, double& X, double& Y, double& Z);
+	static void BreakInputActionValue(FInputActionValue InActionValue, double& X, double& Y, double& Z, EInputActionValueType& Type);
 
 	/**
 	 * Builds an ActionValue from X, Y, Z. Inherits type from an existing ActionValue. Ignores axis values unused by the provided value type.
 	 * @note Intended for use in Input Modifier Modify Raw overloads to modify an existing Input Action Value.
 	 */
 	UFUNCTION(BlueprintPure, Category = "Input", meta = (Keywords = "construct build", NativeMakeFunc))
-	static FInputActionValue MakeInputActionValue(double X, double Y, double Z, const FInputActionValue& MatchValueType);
+	static FInputActionValue MakeInputActionValueOfType(double X, double Y, double Z, const EInputActionValueType ValueType);
 
+	UE_DEPRECATED(5.1, "This version of MakeInputActionValue has been deprecated, please use MakeInputActionValueOfType")
+	UFUNCTION(BlueprintPure, Category = "Input", meta = (Keywords = "construct build", NativeMakeFunc, DeprecatedFunction, DeprecatedMessage="This version of MakeInputActionValue has been deprecated, please use MakeInputActionValueOfType"))
+	static FInputActionValue MakeInputActionValue(double X, double Y, double Z, const FInputActionValue& MatchValueType);
 
 	// Internal helper functionality
 
diff --git a/Engine/Plugins/EnhancedInput/Source/EnhancedInput/Public/InputActionValue.h b/Engine/Plugins/EnhancedInput/Source/EnhancedInput/Public/InputActionValue.h
--- a/Engine/Plugins/EnhancedInput/Source/EnhancedInput/Public/InputActionValue.h
+++ b/Engine/Plugins/EnhancedInput/Source/EnhancedInput/Public/InputActionValue.h
@@ -20,7 +20,7 @@ enum class EInputActionValueType : uint8
 };
 
 
-USTRUCT(BlueprintType, meta = (HasNativeMake = "/Script/EnhancedInput.EnhancedInputLibrary.MakeInputActionValue", HasNativeBreak = "/Script/EnhancedInput.EnhancedInputLibrary.BreakInputActionValue"))
+USTRUCT(BlueprintType, meta = (HasNativeMake = "/Script/EnhancedInput.EnhancedInputLibrary.MakeInputActionValueOfType", HasNativeBreak = "/Script/EnhancedInput.EnhancedInputLibrary.BreakInputActionValue"))
 struct FInputActionValue
 {
 	GENERATED_BODY()
