Early out in QueryKeysMappedToAction if the given action is null, there is no point to iterating the action mappings in this case.

#rnx

diff --git a/Engine/Plugins/EnhancedInput/Source/EnhancedInput/Private/EnhancedInputSubsystemInterface.cpp b/Engine/Plugins/EnhancedInput/Source/EnhancedInput/Private/EnhancedInputSubsystemInterface.cpp
--- a/Engine/Plugins/EnhancedInput/Source/EnhancedInput/Private/EnhancedInputSubsystemInterface.cpp
+++ b/Engine/Plugins/EnhancedInput/Source/EnhancedInput/Private/EnhancedInputSubsystemInterface.cpp
@@ -362,13 +362,16 @@ TArray<FKey> IEnhancedInputSubsystemInterface::QueryKeysMappedToAction(const UIn
 {
 	TArray<FKey> MappedKeys;
 
-	if (const UEnhancedPlayerInput* const PlayerInput = GetPlayerInput())
+	if (Action)
 	{
-		for (const FEnhancedActionKeyMapping& Mapping : PlayerInput->EnhancedActionMappings)
+		if (const UEnhancedPlayerInput* const PlayerInput = GetPlayerInput())
 		{
-			if (Mapping.Action == Action)
+			for (const FEnhancedActionKeyMapping& Mapping : PlayerInput->EnhancedActionMappings)
 			{
-				MappedKeys.AddUnique(Mapping.Key);
+				if (Mapping.Action == Action)
+				{
+					MappedKeys.AddUnique(Mapping.Key);
+				}
 			}
 		}
 	}
