Unshelved from pending changelist '21925691':

Add a double click action to the "Get Value" node in enhanced input to jump to the associated Input Action asset.

#rnx

diff --git a/Engine/Plugins/EnhancedInput/Source/InputBlueprintNodes/Private/K2Node_GetInputActionValue.cpp b/Engine/Plugins/EnhancedInput/Source/InputBlueprintNodes/Private/K2Node_GetInputActionValue.cpp
--- a/Engine/Plugins/EnhancedInput/Source/InputBlueprintNodes/Private/K2Node_GetInputActionValue.cpp
+++ b/Engine/Plugins/EnhancedInput/Source/InputBlueprintNodes/Private/K2Node_GetInputActionValue.cpp
@@ -20,7 +20,7 @@
 #include "K2Node_InputActionValueAccessor.h"
 #include "K2Node_TemporaryVariable.h"
 #include "K2Node_AssignmentStatement.h"
-
+#include "Subsystems/AssetEditorSubsystem.h"
 
 #define LOCTEXT_NAMESPACE "K2Node_GetInputActionValue"
 
@@ -175,6 +175,16 @@ bool UK2Node_GetInputActionValue::IsCompatibleWithGraph(UEdGraph const* Graph) c
 	return (Blueprint != nullptr) && Blueprint->SupportsInputEvents() && !bIsConstructionScript && Super::IsCompatibleWithGraph(Graph);
 }
 
+UObject* UK2Node_GetInputActionValue::GetJumpTargetForDoubleClick() const
+{
+	return const_cast<UObject*>(Cast<UObject>(InputAction));
+}
+
+void UK2Node_GetInputActionValue::JumpToDefinition() const
+{
+	GEditor->GetEditorSubsystem<UAssetEditorSubsystem>()->OpenEditorForAsset(GetJumpTargetForDoubleClick());
+}
+
 void UK2Node_GetInputActionValue::ValidateNodeDuringCompilation(class FCompilerResultsLog& MessageLog) const
 {
 	Super::ValidateNodeDuringCompilation(MessageLog);
diff --git a/Engine/Plugins/EnhancedInput/Source/InputBlueprintNodes/Public/K2Node_GetInputActionValue.h b/Engine/Plugins/EnhancedInput/Source/InputBlueprintNodes/Public/K2Node_GetInputActionValue.h
--- a/Engine/Plugins/EnhancedInput/Source/InputBlueprintNodes/Public/K2Node_GetInputActionValue.h
+++ b/Engine/Plugins/EnhancedInput/Source/InputBlueprintNodes/Public/K2Node_GetInputActionValue.h
@@ -27,6 +27,8 @@ class UK2Node_GetInputActionValue : public UK2Node
 	virtual FText GetNodeTitle(ENodeTitleType::Type TitleType) const override;
 	virtual FText GetTooltipText() const override;
 	virtual bool IsCompatibleWithGraph(UEdGraph const* Graph) const override;
+	virtual UObject* GetJumpTargetForDoubleClick() const override;
+	virtual void JumpToDefinition() const override;
 	//~ End EdGraphNode Interface
 
 	//~ Begin UK2Node Interface
