Add a tooltip to the Get Action Value node in Enhanced Input to include that the input may be 0 if you have certain triggers whose value gets cleared at the end of tick.

#rnx

diff --git a/Engine/Plugins/EnhancedInput/Source/InputBlueprintNodes/Private/K2Node_GetInputActionValue.cpp b/Engine/Plugins/EnhancedInput/Source/InputBlueprintNodes/Private/K2Node_GetInputActionValue.cpp
--- a/Engine/Plugins/EnhancedInput/Source/InputBlueprintNodes/Private/K2Node_GetInputActionValue.cpp
+++ b/Engine/Plugins/EnhancedInput/Source/InputBlueprintNodes/Private/K2Node_GetInputActionValue.cpp
@@ -160,7 +160,9 @@ FText UK2Node_GetInputActionValue::GetTooltipText() const
 	{
 		// FText::Format() is slow, so we cache this to save on performance
 		FString ActionPath = InputAction ? InputAction->GetFullName() : TEXT("");
-		CachedTooltip.SetCachedText(FText::Format(LOCTEXT("GetInputAction_Tooltip", "Returns the current value of {0}.  If input is disabled for the actor the value will be 0."), FText::FromString(ActionPath)), this);
+		CachedTooltip.SetCachedText(
+			FText::Format(LOCTEXT("GetInputAction_Tooltip", "Returns the current value of {0}.  If input is disabled for the actor the value will be 0. \n\nNote: If the value is being altered by an Input Trigger or Input Modifier (such as a Released trigger) then the value may be 0."),
+			FText::FromString(ActionPath)), this);
 	}
 	return CachedTooltip;
 }
