Enhanced Input: Add class picker for Input Action and Input Mapping Context factories (Contributed by [REDACTED_PERSON])

#9557

diff --git a/Engine/Plugins/EnhancedInput/Source/InputEditor/Private/InputCustomizations.cpp b/Engine/Plugins/EnhancedInput/Source/InputEditor/Private/InputCustomizations.cpp
--- a/Engine/Plugins/EnhancedInput/Source/InputEditor/Private/InputCustomizations.cpp
+++ b/Engine/Plugins/EnhancedInput/Source/InputEditor/Private/InputCustomizations.cpp
@@ -289,6 +289,57 @@ TArray<UObject*> FEnhancedInputDeveloperSettingsCustomization::GatherClassDetail
 	return CDOs;
 }
 
+bool FEnhancedInputDeveloperSettingsCustomization::DoesClassHaveSubtypes(UClass* Class)
+{
+	// Search native classes
+	for (TObjectIterator<UClass> ClassIt; ClassIt; ++ClassIt)
+	{
+		// Make sure it's a native class and a child of Class passed in
+		if (ClassIt->IsNative() && ClassIt->IsChildOf(Class))
+		{
+			// Make sure it doesn't have any flags that would disqualify it from being used
+			if (!ClassIt->HasAnyClassFlags(CLASS_Abstract | CLASS_HideDropDown | CLASS_Deprecated | CLASS_NewerVersionExists))
+			{
+				UObject* CDO = ClassIt->GetDefaultObject();
+		
+				// Make Sure it isn't the Class itself
+				if (CDO && CDO->GetClass() != Class)
+				{
+					return true;
+				}
+			}
+		}
+	}
+
+	// Search BPs via asset registry
+	FAssetRegistryModule& AssetRegistryModule = FModuleManager::LoadModuleChecked<FAssetRegistryModule>(FName("AssetRegistry"));
+	IAssetRegistry& AssetRegistry = AssetRegistryModule.Get();
+
+	FARFilter Filter;
+	Filter.ClassPaths.Add(UBlueprint::StaticClass()->GetClassPathName());
+	Filter.bRecursiveClasses = true;
+	TArray<FAssetData> BlueprintAssetData;
+	AssetRegistry.GetAssets(Filter, BlueprintAssetData);
+
+	for (FAssetData& Asset : BlueprintAssetData)
+	{
+		FAssetDataTagMapSharedView::FFindTagResult Result = Asset.TagsAndValues.FindTag(TEXT("NativeParentClass"));
+		if (Result.IsSet() && !ExcludedAssetNames.Contains(Asset.AssetName))
+		{
+			const FString ClassObjectPath = FPackageName::ExportTextPathToObjectPath(Result.GetValue());
+			if (UClass* ParentClass = FindObjectSafe<UClass>(nullptr, *ClassObjectPath, true))
+			{
+				if (ParentClass->IsChildOf(Class))
+				{
+					return true;
+				}
+			}
+		}
+	}
+
+	return false;
+}
+
 void FEnhancedInputDeveloperSettingsCustomization::RebuildDetailsViewForAsset(const FAssetData& AssetData, const bool bIsAssetBeingRemoved)
 {
 	// If the asset was a blueprint...
diff --git a/Engine/Plugins/EnhancedInput/Source/InputEditor/Private/InputCustomizations.h b/Engine/Plugins/EnhancedInput/Source/InputEditor/Private/InputCustomizations.h
--- a/Engine/Plugins/EnhancedInput/Source/InputEditor/Private/InputCustomizations.h
+++ b/Engine/Plugins/EnhancedInput/Source/InputEditor/Private/InputCustomizations.h
@@ -64,6 +64,9 @@ public:
 	virtual void CustomizeDetails(const TSharedPtr<IDetailLayoutBuilder>& DetailBuilder) override;
 	//~ End IDetailCustomization interface
 	
+	/** Returns true if something inherits from the UClass passed in. Checks BPs and Native */
+	static bool DoesClassHaveSubtypes(UClass* Class);
+
 private:
 
 	/** Gather all of the CDO's for the given class, Native and Blueprint. */
diff --git a/Engine/Plugins/EnhancedInput/Source/InputEditor/Private/InputEditorModule.cpp b/Engine/Plugins/EnhancedInput/Source/InputEditor/Private/InputEditorModule.cpp
--- a/Engine/Plugins/EnhancedInput/Source/InputEditor/Private/InputEditorModule.cpp
+++ b/Engine/Plugins/EnhancedInput/Source/InputEditor/Private/InputEditorModule.cpp
@@ -33,6 +33,10 @@
 #include "Styling/SlateStyleRegistry.h"
 #include "ContentBrowserModule.h"
 #include "IContentBrowserSingleton.h"
+#include "ClassViewerModule.h"
+#include "ClassViewerFilter.h"
+#include "Kismet2/KismetEditorUtilities.h"
+#include "Kismet2/SClassPickerDialog.h"
 #include "GameFramework/InputSettings.h"
 #include "EnhancedInputComponent.h"
 #include "EnhancedPlayerInput.h"
@@ -58,6 +62,31 @@ namespace UE::Input
 
 IMPLEMENT_MODULE(FInputEditorModule, InputEditor)
 
+class FInputClassParentFilter : public IClassViewerFilter
+{
+public:
+	FInputClassParentFilter()
+		: DisallowedClassFlags(CLASS_Abstract | CLASS_Deprecated | CLASS_NewerVersionExists | CLASS_HideDropDown) {}
+
+	/** All children of these classes will be included unless filtered out by another setting. */
+	TSet<const UClass*> AllowedChildrenOfClasses;
+
+	/** Disallowed class flags. */
+	EClassFlags DisallowedClassFlags;
+
+	virtual bool IsClassAllowed(const FClassViewerInitializationOptions& InInitOptions, const UClass* InClass, TSharedRef<FClassViewerFilterFuncs> InFilterFuncs) override
+	{
+		return !InClass->HasAnyClassFlags(DisallowedClassFlags)
+			&& InFilterFuncs->IfInChildOfClassesSet(AllowedChildrenOfClasses, InClass) != EFilterReturn::Failed;
+	}
+
+	virtual bool IsUnloadedClassAllowed(const FClassViewerInitializationOptions& InInitOptions, const TSharedRef<const IUnloadedBlueprintData> InUnloadedClassData, TSharedRef<FClassViewerFilterFuncs> InFilterFuncs) override
+	{
+		return !InUnloadedClassData->HasAnyClassFlags(DisallowedClassFlags)
+			&& InFilterFuncs->IfInChildOfClassesSet(AllowedChildrenOfClasses, InUnloadedClassData) != EFilterReturn::Failed;
+	}
+};
+
 // Asset factories
 
 // InputContext
@@ -67,10 +96,55 @@ UInputMappingContext_Factory::UInputMappingContext_Factory(const class FObjectIn
 	bCreateNew = true;
 }
 
+bool UInputMappingContext_Factory::ConfigureProperties()
+{
+	if (!FEnhancedInputDeveloperSettingsCustomization::DoesClassHaveSubtypes(UInputMappingContext::StaticClass()))
+	{
+		return true;
+	}
+	
+	// nullptr the InputMappingContextClass so we can check for selection
+	InputMappingContextClass = nullptr;
+
+	// Load the classviewer module to display a class picker
+	FClassViewerModule& ClassViewerModule = FModuleManager::LoadModuleChecked<FClassViewerModule>("ClassViewer");
+
+	// Fill in options
+	FClassViewerInitializationOptions Options;
+	Options.Mode = EClassViewerMode::ClassPicker;
+
+	TSharedPtr<FInputClassParentFilter> Filter = MakeShareable(new FInputClassParentFilter);
+	Filter->AllowedChildrenOfClasses.Add(UInputMappingContext::StaticClass());
+
+	Options.ClassFilters.Add(Filter.ToSharedRef());
+
+	const FText TitleText = LOCTEXT("CreateInputMappingContextOptions", "Pick Class For Input Mapping Context Instance");
+	UClass* ChosenClass = nullptr;
+	const bool bPressedOk = SClassPickerDialog::PickClass(TitleText, Options, ChosenClass, UInputMappingContext::StaticClass());
+
+	if (bPressedOk)
+	{
+		InputMappingContextClass = ChosenClass;
+	}
+
+	return bPressedOk;
+}
+
 UObject* UInputMappingContext_Factory::FactoryCreateNew(UClass* Class, UObject* InParent, FName Name, EObjectFlags Flags, UObject* Context, FFeedbackContext* Warn)
 {
-	check(Class->IsChildOf(UInputMappingContext::StaticClass()));
-	UInputMappingContext* IMC = NewObject<UInputMappingContext>(InParent, Class, Name, Flags | RF_Transactional, Context);
+	UInputMappingContext* IMC = nullptr;
+
+	if (InputMappingContextClass != nullptr)
+	{
+		IMC = NewObject<UInputMappingContext>(InParent, InputMappingContextClass, Name, Flags | RF_Transactional, Context);
+	}
+	else
+	{
+		check(Class->IsChildOf(UInputMappingContext::StaticClass()));
+		IMC = NewObject<UInputMappingContext>(InParent, Class, Name, Flags | RF_Transactional, Context);
+	}
+
+	check(IMC);
 
 	// Populate the IMC with some initial input actions if they were specified. This will be the case if the IMC is being created from the FAssetTypeActions_InputAction
 	for (TWeakObjectPtr<UInputAction> WeakIA : InitialActions)
@@ -99,10 +173,51 @@ UInputAction_Factory::UInputAction_Factory(const class FObjectInitializer& OBJ)
 	bCreateNew = true;
 }
 
-UObject* UInputAction_Factory::FactoryCreateNew(UClass* Class, UObject* InParent, FName Name, EObjectFlags Flags, UObject* Context, FFeedbackContext* Warn)
+bool UInputAction_Factory::ConfigureProperties()
 {
-	check(Class->IsChildOf(UInputAction::StaticClass()));
-	return NewObject<UInputAction>(InParent, Class, Name, Flags | RF_Transactional, Context);
+	if (!FEnhancedInputDeveloperSettingsCustomization::DoesClassHaveSubtypes(UInputAction::StaticClass()))
+	{
+		return true;
+	}
+	
+	// nullptr the InputActionClass so we can check for selection
+	InputActionClass = nullptr;
+
+	// Load the classviewer module to display a class picker
+	FClassViewerModule& ClassViewerModule = FModuleManager::LoadModuleChecked<FClassViewerModule>("ClassViewer");
+
+	// Fill in options
+	FClassViewerInitializationOptions Options;
+	Options.Mode = EClassViewerMode::ClassPicker;
+
+	TSharedPtr<FInputClassParentFilter> Filter = MakeShareable(new FInputClassParentFilter);
+	Filter->AllowedChildrenOfClasses.Add(UInputAction::StaticClass());
+
+	Options.ClassFilters.Add(Filter.ToSharedRef());
+
+	const FText TitleText = LOCTEXT("CreateInputActionOptions", "Pick Class For Input Action Instance");
+	UClass* ChosenClass = nullptr;
+	const bool bPressedOk = SClassPickerDialog::PickClass(TitleText, Options, ChosenClass, UInputAction::StaticClass());
+
+	if (bPressedOk)
+	{
+		InputActionClass = ChosenClass;
+	}
+
+	return bPressedOk;
+}
+
+UObject* UInputAction_Factory::FactoryCreateNew(UClass* Class, UObject* InParent, FName Name, EObjectFlags Flags, UObject* Context, FFeedbackContext* Warn)
+{	
+	if (InputActionClass != nullptr)
+	{
+		return NewObject<UInputAction>(InParent, InputActionClass, Name, Flags | RF_Transactional, Context);
+	}
+	else
+	{
+		check(Class->IsChildOf(UInputAction::StaticClass()));
+		return NewObject<UInputAction>(InParent, Class, Name, Flags | RF_Transactional, Context);
+	}
 }
 
 // UPlayerMappableInputConfig_Factory
diff --git a/Engine/Plugins/EnhancedInput/Source/InputEditor/Public/InputEditorModule.h b/Engine/Plugins/EnhancedInput/Source/InputEditor/Public/InputEditorModule.h
--- a/Engine/Plugins/EnhancedInput/Source/InputEditor/Public/InputEditorModule.h
+++ b/Engine/Plugins/EnhancedInput/Source/InputEditor/Public/InputEditorModule.h
@@ -63,6 +63,10 @@ class INPUTEDITOR_API UInputMappingContext_Factory : public UFactory
 {
 	GENERATED_UCLASS_BODY()
 public:
+	UPROPERTY(EditAnywhere, Category=InputMappingContext)
+	TSubclassOf<class UInputMappingContext> InputMappingContextClass;
+
+	virtual bool ConfigureProperties() override;
 	virtual UObject* FactoryCreateNew(UClass* Class, UObject* InParent, FName Name, EObjectFlags Flags, UObject* Context, FFeedbackContext* Warn) override;
 
 	/** Set the array of initial actions that the resulting IMC should be populated with */
@@ -80,6 +84,10 @@ class INPUTEDITOR_API UInputAction_Factory : public UFactory
 {
 	GENERATED_UCLASS_BODY()
 public:
+	UPROPERTY(EditAnywhere, Category=InputAction)
+	TSubclassOf<UInputAction> InputActionClass;
+
+	virtual bool ConfigureProperties() override;
 	virtual UObject* FactoryCreateNew(UClass* Class, UObject* InParent, FName Name, EObjectFlags Flags, UObject* Context, FFeedbackContext* Warn) override;
 };
 
