Made ClearActionBindings and ClearBindingsForObject virtual on UInputComponent and overrode on UEnhancedInputComponent to ensure the additional Enhanced action bindings are cleared as well

diff --git a/Engine/Plugins/EnhancedInput/Source/EnhancedInput/Private/EnhancedInputComponent.cpp b/Engine/Plugins/EnhancedInput/Source/EnhancedInput/Private/EnhancedInputComponent.cpp
--- a/Engine/Plugins/EnhancedInput/Source/EnhancedInput/Private/EnhancedInputComponent.cpp
+++ b/Engine/Plugins/EnhancedInput/Source/EnhancedInput/Private/EnhancedInputComponent.cpp
@@ -34,6 +34,26 @@ bool UEnhancedInputComponent::HasBindings( ) const
 	return GetActionEventBindings().Num() > 0 || GetActionValueBindings().Num() > 0 || GetDebugKeyBindings().Num() > 0 || Super::HasBindings();
 }
 
+void UEnhancedInputComponent::ClearActionBindings()
+{
+	Super::ClearActionBindings();
+
+	ClearActionEventBindings();
+}
+
+void UEnhancedInputComponent::ClearBindingsForObject(UObject* InOwner)
+{
+	Super::ClearBindingsForObject(InOwner);
+
+	for (int32 Index = EnhancedActionEventBindings.Num() - 1; Index >= 0; --Index)
+	{
+		if (EnhancedActionEventBindings[Index]->IsBoundToObject(InOwner))
+		{
+			EnhancedActionEventBindings.RemoveAtSwap(Index, 1, false);
+		}
+	}
+}
+
 template<typename T>
 bool RemoveBindingByIndex(TArray<T>& Bindings, const int32 BindingIndex)
 {
diff --git a/Engine/Plugins/EnhancedInput/Source/EnhancedInput/Public/EnhancedInputComponent.h b/Engine/Plugins/EnhancedInput/Source/EnhancedInput/Public/EnhancedInputComponent.h
--- a/Engine/Plugins/EnhancedInput/Source/EnhancedInput/Public/EnhancedInputComponent.h
+++ b/Engine/Plugins/EnhancedInput/Source/EnhancedInput/Public/EnhancedInputComponent.h
@@ -171,6 +171,7 @@ public:
 	virtual void Execute(const FInputActionInstance& ActionData) const = 0;
 	virtual TUniquePtr<FEnhancedInputActionEventBinding> Clone() const = 0;
 	virtual void SetShouldFireWithEditorScriptGuard(const bool bNewValue) = 0;
+	virtual bool IsBoundToObject(void const* Object) const = 0;
 };
 
 /** Binds an action value for later reference. CurrentValue will be kept up to date with the value of the bound action */
@@ -260,6 +261,11 @@ public:
 		Delegate.SetShouldFireWithEditorScriptGuard(bNewValue);
 	}
 
+	virtual bool IsBoundToObject(void const* Object) const override
+	{
+		return Delegate.IsBoundToObject(Object);
+	}
+
 	TEnhancedInputUnifiedDelegate<TSignature> Delegate;
 };
 
@@ -367,6 +373,12 @@ public:
 	void ClearActionValueBindings() { EnhancedActionValueBindings.Reset(); }
 	void ClearDebugKeyBindings() { DebugKeyBindings.Reset(); }
 
+	/** Removes all action bindings. */
+	virtual void ClearActionBindings() override;
+	
+	/** Clears any input callback delegates from the given UObject */
+	virtual void ClearBindingsForObject(UObject* InOwner) override;
+
 	/**
 	 * Gets all action bindings of type
 	 *
diff --git a/Engine/Source/Runtime/Engine/Classes/Components/InputComponent.h b/Engine/Source/Runtime/Engine/Classes/Components/InputComponent.h
--- a/Engine/Source/Runtime/Engine/Classes/Components/InputComponent.h
+++ b/Engine/Source/Runtime/Engine/Classes/Components/InputComponent.h
@@ -765,7 +765,7 @@ public:
 	uint8 bBlockInput:1;
 
 	/** Clears any inut callback delegates from the given UObject */
-	void ClearBindingsForObject(UObject* InOwner);
+	virtual void ClearBindingsForObject(UObject* InOwner);
 	
 	void ConditionalBuildKeyMap(UPlayerInput* PlayerInput);
 
@@ -818,7 +818,7 @@ public:
 	 *
 	 * @see AddActionBinding, GetActionBinding, GetNumActionBindings, RemoveActionBinding
 	 */
-	void ClearActionBindings();
+	virtual void ClearActionBindings();
 
 	/**
 	 * Gets the action binding with the specified index.
