Fixing tooltip typo in Affected by [REDACTED_PERSON] mouse over text

#rnx

diff --git a/Engine/Plugins/EnhancedInput/Source/EnhancedInput/Public/InputTriggers.h b/Engine/Plugins/EnhancedInput/Source/EnhancedInput/Public/InputTriggers.h
--- a/Engine/Plugins/EnhancedInput/Source/EnhancedInput/Public/InputTriggers.h
+++ b/Engine/Plugins/EnhancedInput/Source/EnhancedInput/Public/InputTriggers.h
@@ -197,7 +197,7 @@ public:
 	 * Should global time dilation be applied to the held duration?
 	 * Default is set to false.
 	 * 
-	 * If this is set to true, then the owning Player Controller's actor time dialtion
+	 * If this is set to true, then the owning Player Controller's actor time dilation
 	 * will be used when calculating the HeldDuration.
 	 * 
 	 * @see UInputTriggerTimedBase::CalculateHeldDuration
