Change a couple of enhanced input classes to only preload certain assets that were discovered in the asset registry if the plugin they live in is actually loaded

diff --git a/Engine/Plugins/EnhancedInput/Source/InputBlueprintNodes/Private/K2Node_EnhancedInputAction.cpp b/Engine/Plugins/EnhancedInput/Source/InputBlueprintNodes/Private/K2Node_EnhancedInputAction.cpp
--- a/Engine/Plugins/EnhancedInput/Source/InputBlueprintNodes/Private/K2Node_EnhancedInputAction.cpp
+++ b/Engine/Plugins/EnhancedInput/Source/InputBlueprintNodes/Private/K2Node_EnhancedInputAction.cpp
@@ -424,10 +424,13 @@ void UK2Node_EnhancedInputAction::GetMenuActions(FBlueprintActionDatabaseRegistr
 			UBlueprintNodeSpawner* NodeSpawner = UBlueprintNodeSpawner::Create(GetClass());
 			check(NodeSpawner != nullptr);
 
-			if (const UInputAction* Action = Cast<const UInputAction>(ActionAsset.GetAsset()))
+			if (FPackageName::GetPackageMountPoint(ActionAsset.PackageName.ToString()) != NAME_None)
 			{
-				NodeSpawner->CustomizeNodeDelegate = UBlueprintNodeSpawner::FCustomizeNodeDelegate::CreateStatic(CustomizeInputNodeLambda, TWeakObjectPtr<const UInputAction>(Action));
-				ActionRegistrar.AddBlueprintAction(Action, NodeSpawner);	
+				if (const UInputAction* Action = Cast<const UInputAction>(ActionAsset.GetAsset()))
+				{
+					NodeSpawner->CustomizeNodeDelegate = UBlueprintNodeSpawner::FCustomizeNodeDelegate::CreateStatic(CustomizeInputNodeLambda, TWeakObjectPtr<const UInputAction>(Action));
+					ActionRegistrar.AddBlueprintAction(Action, NodeSpawner);
+				}
 			}
 		}
 	}
diff --git a/Engine/Plugins/EnhancedInput/Source/InputBlueprintNodes/Private/K2Node_GetInputActionValue.cpp b/Engine/Plugins/EnhancedInput/Source/InputBlueprintNodes/Private/K2Node_GetInputActionValue.cpp
--- a/Engine/Plugins/EnhancedInput/Source/InputBlueprintNodes/Private/K2Node_GetInputActionValue.cpp
+++ b/Engine/Plugins/EnhancedInput/Source/InputBlueprintNodes/Private/K2Node_GetInputActionValue.cpp
@@ -227,10 +227,13 @@ void UK2Node_GetInputActionValue::GetMenuActions(FBlueprintActionDatabaseRegistr
 			UBlueprintNodeSpawner* NodeSpawner = UBlueprintNodeSpawner::Create(GetClass());
 			check(NodeSpawner != nullptr);
 
-			if (const UInputAction* Action = Cast<const UInputAction>(ActionAsset.GetAsset()))
+			if (FPackageName::GetPackageMountPoint(ActionAsset.PackageName.ToString()) != NAME_None)
 			{
-				NodeSpawner->CustomizeNodeDelegate = UBlueprintNodeSpawner::FCustomizeNodeDelegate::CreateStatic(CustomizeInputNodeLambda, TWeakObjectPtr<const UInputAction>(Action));
-				ActionRegistrar.AddBlueprintAction(Action, NodeSpawner);
+				if (const UInputAction* Action = Cast<const UInputAction>(ActionAsset.GetAsset()))
+				{
+					NodeSpawner->CustomizeNodeDelegate = UBlueprintNodeSpawner::FCustomizeNodeDelegate::CreateStatic(CustomizeInputNodeLambda, TWeakObjectPtr<const UInputAction>(Action));
+					ActionRegistrar.AddBlueprintAction(Action, NodeSpawner);
+				}
 			}
 		}
 	}
