Renaming Combo Trigger Input Trigger to fix a compile error when opening Lyra

#rnx

diff --git a/Engine/Plugins/EnhancedInput/Source/EnhancedInput/Private/InputTriggers.cpp b/Engine/Plugins/EnhancedInput/Source/EnhancedInput/Private/InputTriggers.cpp
--- a/Engine/Plugins/EnhancedInput/Source/EnhancedInput/Private/InputTriggers.cpp
+++ b/Engine/Plugins/EnhancedInput/Source/EnhancedInput/Private/InputTriggers.cpp
@@ -228,16 +228,16 @@ ETriggerState UInputTriggerChordAction::UpdateState_Implementation(const UEnhanc
 	return EventData ? EventData->TriggerStateTracker.GetState() : ETriggerState::None;
 }
 
-UInputTriggerComboAction::UInputTriggerComboAction()
+UInputTriggerCombo::UInputTriggerCombo()
 {
 	bShouldAlwaysTick = true;
 }
 
-ETriggerState UInputTriggerComboAction::UpdateState_Implementation(const UEnhancedPlayerInput* PlayerInput, FInputActionValue ModifiedValue, float DeltaTime)
+ETriggerState UInputTriggerCombo::UpdateState_Implementation(const UEnhancedPlayerInput* PlayerInput, FInputActionValue ModifiedValue, float DeltaTime)
 {
 	if (ComboActions.IsEmpty())
 	{
-		ensureMsgf(false, TEXT("You must add combo actions to the UInputTriggerComboAction for it work properly! Exiting..."));
+		ensureMsgf(false, TEXT("You must add combo actions to the UInputTriggerCombo for it work properly! Exiting..."));
 		return ETriggerState::None;
 	}
 	
@@ -316,7 +316,7 @@ ETriggerState UInputTriggerComboAction::UpdateState_Implementation(const UEnhanc
 };
 
 #if WITH_EDITOR
-EDataValidationResult UInputTriggerComboAction::IsDataValid(TArray<FText>& ValidationErrors)
+EDataValidationResult UInputTriggerCombo::IsDataValid(TArray<FText>& ValidationErrors)
 {
 	EDataValidationResult Result = CombineDataValidationResults(Super::IsDataValid(ValidationErrors), EDataValidationResult::Valid);
 	
diff --git a/Engine/Plugins/EnhancedInput/Source/EnhancedInput/Public/InputTriggers.h b/Engine/Plugins/EnhancedInput/Source/EnhancedInput/Public/InputTriggers.h
--- a/Engine/Plugins/EnhancedInput/Source/EnhancedInput/Public/InputTriggers.h
+++ b/Engine/Plugins/EnhancedInput/Source/EnhancedInput/Public/InputTriggers.h
@@ -428,15 +428,15 @@ struct FInputComboStepData
 };
 
 /**
- * UInputTriggerComboAction
+ * UInputTriggerCombo
  * All actions in the combo array must be pressed within a timeframe to trigger
 */
 UCLASS(NotBlueprintable, meta = (DisplayName = "Combo", NotInputConfigurable = "true"))
-class ENHANCEDINPUT_API UInputTriggerComboAction : public UInputTrigger
+class ENHANCEDINPUT_API UInputTriggerCombo : public UInputTrigger
 {
 	GENERATED_BODY()
 
-	UInputTriggerComboAction();
+	UInputTriggerCombo();
 
 protected:
 	// Implicit, so action cannot fire unless this is firing.
diff --git a/Engine/Plugins/EnhancedInput/Source/InputEditor/Private/InputCustomizations.cpp b/Engine/Plugins/EnhancedInput/Source/InputEditor/Private/InputCustomizations.cpp
--- a/Engine/Plugins/EnhancedInput/Source/InputEditor/Private/InputCustomizations.cpp
+++ b/Engine/Plugins/EnhancedInput/Source/InputEditor/Private/InputCustomizations.cpp
@@ -117,7 +117,7 @@ bool FEnhancedActionMappingCustomization::DoesTriggerArrayContainCombo() const
 					// Make sure we can cast this to a input trigger and if it's a combo we can return true
 					if (UInputTrigger** InputComboTrigger = reinterpret_cast<UInputTrigger**>(ArrayHelper.GetRawPtr(i)))
 					{
-						if ((*InputComboTrigger) && (*InputComboTrigger)->IsA(UInputTriggerComboAction::StaticClass()))
+						if ((*InputComboTrigger) && (*InputComboTrigger)->IsA(UInputTriggerCombo::StaticClass()))
 						{
 							return true;
 						}
