Check to make sure there is a valid world on the Enhanced Input editor subsystem before attempting to input the key. The world can be invalid if you leave the subsystem processing input while you resave different maps in the editor.

#rnx

diff --git a/Engine/Plugins/EnhancedInput/Source/InputEditor/Private/EnhancedInputEditorProcessor.cpp b/Engine/Plugins/EnhancedInput/Source/InputEditor/Private/EnhancedInputEditorProcessor.cpp
--- a/Engine/Plugins/EnhancedInput/Source/InputEditor/Private/EnhancedInputEditorProcessor.cpp
+++ b/Engine/Plugins/EnhancedInput/Source/InputEditor/Private/EnhancedInputEditorProcessor.cpp
@@ -175,7 +175,10 @@ bool FEnhancedInputEditorProcessor::InputKeyToSubsystem(const FInputKeyParams& P
 	{
 		if (UEnhancedInputEditorSubsystem* Subsystem = GEditor->GetEditorSubsystem<UEnhancedInputEditorSubsystem>())
 		{
-			return Subsystem->InputKey(Params);
+			if (Subsystem->GetWorld())
+			{
+				return Subsystem->InputKey(Params);
+			}
 		}	
 	}
 	
