Fix missing Enhanced Player Input UPROPERTY(s) (Contributed by [REDACTED_PERSON])

#rnx
#9649

diff --git a/Engine/Plugins/EnhancedInput/Source/EnhancedInput/Public/EnhancedPlayerInput.h b/Engine/Plugins/EnhancedInput/Source/EnhancedInput/Public/EnhancedPlayerInput.h
--- a/Engine/Plugins/EnhancedInput/Source/EnhancedInput/Public/EnhancedPlayerInput.h
+++ b/Engine/Plugins/EnhancedInput/Source/EnhancedInput/Public/EnhancedPlayerInput.h
@@ -16,6 +16,20 @@ enum class ETriggerEventInternal : uint8;
 enum class EKeyEvent : uint8;
 class UInputMappingContext;
 
+struct FInjectedInput
+{
+	FInputActionValue RawValue;
+	TArray<UInputTrigger*> Triggers;
+	TArray<UInputModifier*> Modifiers;
+};
+
+USTRUCT()
+struct FInjectedInputArray
+{
+	GENERATED_BODY()
+	TArray<FInjectedInput> Injected;
+};
+
 /**
 * UEnhancedPlayerInput : UPlayerInput extensions for enhanced player input system
 */
@@ -116,21 +130,12 @@ private:
 	UPROPERTY(Transient)
 	TMap<FKey, FVector> KeysPressedThisTick;
 
-	struct FInjectedInput
-	{
-		FInputActionValue RawValue;
-		TArray<UInputTrigger*> Triggers;
-		TArray<UInputModifier*> Modifiers;
-	};
-	struct FInjectedInputArray
-	{
-		TArray<FInjectedInput> Injected;
-	};
-
 	/** Inputs injected since the last call to ProcessInputStack */
+	UPROPERTY(Transient)
 	TMap<const UInputAction*, FInjectedInputArray> InputsInjectedThisTick;
 
 	/** Last frame's injected inputs */
+	UPROPERTY(Transient)
 	TSet<const UInputAction*> LastInjectedActions;
 
 	/** Used to keep track of Input Actions that have UInputTriggerChordAction triggers on them */
