Add a log when you call AddMappingContext with a null param, that way users get warnings i they forget to populate a property or something like that.

#rnx

diff --git a/Engine/Plugins/EnhancedInput/Source/EnhancedInput/Private/EnhancedInputSubsystemInterface.cpp b/Engine/Plugins/EnhancedInput/Source/EnhancedInput/Private/EnhancedInputSubsystemInterface.cpp
--- a/Engine/Plugins/EnhancedInput/Source/EnhancedInput/Private/EnhancedInputSubsystemInterface.cpp
+++ b/Engine/Plugins/EnhancedInput/Source/EnhancedInput/Private/EnhancedInputSubsystemInterface.cpp
@@ -65,6 +65,10 @@ void IEnhancedInputSubsystemInterface::AddMappingContext(const UInputMappingCont
 			RequestRebuildControlMappings(Options);
 		}
 	}
+	else
+	{
+		UE_LOG(LogEnhancedInput, Warning, TEXT("Called AddMappingContext with a null Mapping Context! No changes have been applied."));
+	}
 }
 
 void IEnhancedInputSubsystemInterface::RemoveMappingContext(const UInputMappingContext* MappingContext, const FModifyContextOptions& Options)
