Removing some native pointer usage in favor of TObjectPtr in Enhanced Input

#rnx

diff --git a/Engine/Plugins/EnhancedInput/Source/EnhancedInput/Private/EnhancedPlayerInput.cpp b/Engine/Plugins/EnhancedInput/Source/EnhancedInput/Private/EnhancedPlayerInput.cpp
--- a/Engine/Plugins/EnhancedInput/Source/EnhancedInput/Private/EnhancedPlayerInput.cpp
+++ b/Engine/Plugins/EnhancedInput/Source/EnhancedInput/Private/EnhancedPlayerInput.cpp
@@ -121,7 +121,7 @@ enum class EKeyEvent : uint8
 	Held,		// Key generated no event, but is in a held state and wants to continue applying modifiers and triggers
 };
 
-void UEnhancedPlayerInput::ProcessActionMappingEvent(const UInputAction* Action, float DeltaTime, bool bGamePaused, FInputActionValue RawKeyValue, EKeyEvent KeyEvent, const TArray<UInputModifier*>& Modifiers, const TArray<UInputTrigger*>& Triggers)
+void UEnhancedPlayerInput::ProcessActionMappingEvent(TObjectPtr<const UInputAction> Action, float DeltaTime, bool bGamePaused, FInputActionValue RawKeyValue, EKeyEvent KeyEvent, const TArray<UInputModifier*>& Modifiers, const TArray<UInputTrigger*>& Triggers)
 {
 	FInputActionInstance& ActionData = FindOrAddActionEventData(Action);
 
@@ -196,7 +196,7 @@ void UEnhancedPlayerInput::ProcessActionMappingEvent(const UInputAction* Action,
 	ActionData.TriggerStateTracker.SetMappingTriggerApplied(bMappingTriggersApplied);
 }
 
-void UEnhancedPlayerInput::InjectInputForAction(const UInputAction* Action, FInputActionValue RawValue, const TArray<UInputModifier*>& Modifiers, const TArray<UInputTrigger*>& Triggers)
+void UEnhancedPlayerInput::InjectInputForAction(TObjectPtr<const UInputAction> Action, FInputActionValue RawValue, const TArray<UInputModifier*>& Modifiers, const TArray<UInputTrigger*>& Triggers)
 {
 	FInjectedInput Input;
 	Input.RawValue = RawValue;
@@ -329,7 +329,7 @@ void UEnhancedPlayerInput::ProcessInputStack(const TArray<UInputComponent*>& Inp
 	for (auto It = LastInjectedActions.CreateIterator(); It; ++It)
 	{
 		TRACE_CPUPROFILER_EVENT_SCOPE(EnhPIS_InjectedStrip);
-		const UInputAction* InjectedAction = *It;
+		TObjectPtr<const UInputAction> InjectedAction = *It;
 
 		if (!InjectedAction)
 		{
@@ -344,10 +344,10 @@ void UEnhancedPlayerInput::ProcessInputStack(const TArray<UInputComponent*>& Inp
 	}
 
 	// Handle injected inputs, applying modifiers and triggers
-	for (TPair<const UInputAction*, FInjectedInputArray>& InjectedPair : InputsInjectedThisTick)
+	for (TPair<TObjectPtr<const UInputAction>, FInjectedInputArray>& InjectedPair : InputsInjectedThisTick)
 	{
 		TRACE_CPUPROFILER_EVENT_SCOPE(EnhPIS_Injected);
-		const UInputAction* InjectedAction = InjectedPair.Key;
+		TObjectPtr<const UInputAction> InjectedAction = InjectedPair.Key;
 		if (!InjectedAction)
 		{
 			continue;
@@ -372,7 +372,7 @@ void UEnhancedPlayerInput::ProcessInputStack(const TArray<UInputComponent*>& Inp
 	{
 		TRACE_CPUPROFILER_EVENT_SCOPE(EnhPIS_PostTick);
 
-		const UInputAction* Action = ActionPair.Key;
+		TObjectPtr<const UInputAction> Action = ActionPair.Key;
 		FInputActionInstance& ActionData = ActionPair.Value;
 		ETriggerState TriggerState = ETriggerState::None;
 
@@ -470,7 +470,7 @@ void UEnhancedPlayerInput::ProcessInputStack(const TArray<UInputComponent*>& Inp
 		// Action all delegates that triggered this tick, in the order in which they triggered.
 		for (TUniquePtr<FEnhancedInputActionEventBinding>& Delegate : TriggeredDelegates)
 		{
-			const UInputAction* DelegateAction = Delegate->GetAction();
+			TObjectPtr<const UInputAction> DelegateAction = Delegate->GetAction();
 			bool bCanTrigger = true;
 
 			if (UE::Input::ShouldOnlyTriggerLastActionInChord)
@@ -606,7 +606,7 @@ void UEnhancedPlayerInput::ProcessInputStack(const TArray<UInputComponent*>& Inp
 	KeysPressedThisTick.Reset();
 }
 
-FInputActionValue UEnhancedPlayerInput::GetActionValue(const UInputAction* ForAction) const
+FInputActionValue UEnhancedPlayerInput::GetActionValue(TObjectPtr<const UInputAction> ForAction) const
 {
 	const FInputActionInstance* ActionData = FindActionInstanceData(ForAction);
 	return ActionData ? ActionData->GetValue() : FInputActionValue(ForAction->ValueType, FInputActionValue::Axis3D::ZeroVector);
@@ -631,11 +631,11 @@ void UEnhancedPlayerInput::ClearAllMappings()
 }
 
 template<typename T>
-void UEnhancedPlayerInput::GatherActionEventDataForActionMap(const T& ActionMap, TMap<const UInputAction*, FInputActionInstance>& FoundActionEventData) const
+void UEnhancedPlayerInput::GatherActionEventDataForActionMap(const T& ActionMap, TMap<TObjectPtr<const UInputAction>, FInputActionInstance>& FoundActionEventData) const
 {
 	for (const typename T::ElementType& Pair : ActionMap)
 	{
-		const UInputAction* Action = Pair.Key;
+		TObjectPtr<const UInputAction> Action = Pair.Key;
 		if (FInputActionInstance* ActionData = ActionInstanceData.Find(Action))
 		{
 			FoundActionEventData.Add(Action, *ActionData);
@@ -650,7 +650,7 @@ void UEnhancedPlayerInput::ConditionalBuildKeyMappings_Internal() const
 	// Remove any ActionEventData without a corresponding entry in EnhancedActionMappings or the injection maps
 	for (auto Itr = ActionInstanceData.CreateIterator(); Itr; ++Itr)
 	{
-		const UInputAction* Action = Itr.Key();
+		TObjectPtr<const UInputAction> Action = Itr.Key();
 
 		auto HasActionMapping = [&Action](const FEnhancedActionKeyMapping& Mapping) { return Mapping.Action == Action; };
 
@@ -688,7 +688,7 @@ bool UEnhancedPlayerInput::IsKeyHandledByAction(FKey Key) const
 	return EnhancedKeyBinds.Contains(Key) || Super::IsKeyHandledByAction(Key);
 }
 
-FInputActionInstance& UEnhancedPlayerInput::FindOrAddActionEventData(const UInputAction* Action) const
+FInputActionInstance& UEnhancedPlayerInput::FindOrAddActionEventData(TObjectPtr<const UInputAction> Action) const
 {
 	FInputActionInstance* Instance = ActionInstanceData.Find(Action);
 	if (!Instance)
diff --git a/Engine/Plugins/EnhancedInput/Source/EnhancedInput/Public/EnhancedPlayerInput.h b/Engine/Plugins/EnhancedInput/Source/EnhancedInput/Public/EnhancedPlayerInput.h
--- a/Engine/Plugins/EnhancedInput/Source/EnhancedInput/Public/EnhancedPlayerInput.h
+++ b/Engine/Plugins/EnhancedInput/Source/EnhancedInput/Public/EnhancedPlayerInput.h
@@ -47,15 +47,15 @@ public:
 	/**
 	* Returns the action instance data for the given input action if there is any. Returns nullptr if the action is not available.
 	*/
-	const FInputActionInstance* FindActionInstanceData(const UInputAction* ForAction) const { return ActionInstanceData.Find(ForAction); }
+	const FInputActionInstance* FindActionInstanceData(TObjectPtr<const UInputAction> ForAction) const { return ActionInstanceData.Find(ForAction); }
 
 	/** Retrieve the current value of an action for this player.
 	* Note: If the action is not currently triggering this will return a zero value of the appropriate value type, ignoring any ongoing inputs.
 	*/
-	FInputActionValue GetActionValue(const UInputAction* ForAction) const;
+	FInputActionValue GetActionValue(TObjectPtr<const UInputAction> ForAction) const;
 
 	// Input simulation via injection. Runs modifiers and triggers delegates as if the input had come through the underlying input system as FKeys. Applies action modifiers and triggers on top.
-	void InjectInputForAction(const UInputAction* Action, FInputActionValue RawValue, const TArray<UInputModifier*>& Modifiers = {}, const TArray<UInputTrigger*>& Triggers = {});
+	void InjectInputForAction(TObjectPtr<const UInputAction> Action, FInputActionValue RawValue, const TArray<UInputModifier*>& Modifiers = {}, const TArray<UInputTrigger*>& Triggers = {});
 
 	virtual bool InputKey(const FInputKeyParams& Params) override;
 	
@@ -93,12 +93,12 @@ private:
 	FInputActionValue ApplyModifiers(const TArray<UInputModifier*>& Modifiers, FInputActionValue RawValue, float DeltaTime) const;						// Pre-modified (raw) value
 	ETriggerEventInternal GetTriggerStateChangeEvent(ETriggerState LastTriggerState, ETriggerState NewTriggerState) const;
 	ETriggerEvent ConvertInternalTriggerEvent(ETriggerEventInternal Event) const;	// Collapse a detailed internal trigger event into a friendly representation
-	void ProcessActionMappingEvent(const UInputAction* Action, float DeltaTime, bool bGamePaused, FInputActionValue RawValue, EKeyEvent KeyEvent, const TArray<UInputModifier*>& Modifiers, const TArray<UInputTrigger*>& Triggers);
+	void ProcessActionMappingEvent(TObjectPtr<const UInputAction> Action, float DeltaTime, bool bGamePaused, FInputActionValue RawValue, EKeyEvent KeyEvent, const TArray<UInputModifier*>& Modifiers, const TArray<UInputTrigger*>& Triggers);
 
-	FInputActionInstance& FindOrAddActionEventData(const UInputAction* Action) const;
+	FInputActionInstance& FindOrAddActionEventData(TObjectPtr<const UInputAction> Action) const;
 
 	template<typename T>
-	void GatherActionEventDataForActionMap(const T& ActionMap, TMap<const UInputAction*, FInputActionInstance>& FoundActionEventData) const;
+	void GatherActionEventDataForActionMap(const T& ActionMap, TMap<TObjectPtr<const UInputAction>, FInputActionInstance>& FoundActionEventData) const;
 
 	/** Currently applied key mappings
 	 * Note: Source reference only. Use EnhancedActionMappings for the actual mappings (with properly instanced triggers/modifiers)
@@ -118,7 +118,7 @@ private:
 	mutable TMap<TObjectPtr<const UInputAction>, FInputActionInstance> ActionInstanceData;
 
 	/** Actions which had actuated events at the last call to ProcessInputStack (held/pressed/released) */
-	TSet<const UInputAction*> ActionsWithEventsThisTick;
+	TSet<TObjectPtr<const UInputAction>> ActionsWithEventsThisTick;
 
 	/** Actions that have been triggered this tick and have a delegate that may be fired */
 	TSet<TObjectPtr<const UInputAction>> TriggeredActionsThisTick;
@@ -132,11 +132,11 @@ private:
 
 	/** Inputs injected since the last call to ProcessInputStack */
 	UPROPERTY(Transient)
-	TMap<const UInputAction*, FInjectedInputArray> InputsInjectedThisTick;
+	TMap<TObjectPtr<const UInputAction>, FInjectedInputArray> InputsInjectedThisTick;
 
 	/** Last frame's injected inputs */
 	UPROPERTY(Transient)
-	TSet<const UInputAction*> LastInjectedActions;
+	TSet<TObjectPtr<const UInputAction>> LastInjectedActions;
 
 	/** Used to keep track of Input Actions that have UInputTriggerChordAction triggers on them */
 	struct FDependentChordTracker
