Create a "Default Input Mapping Contexts" developer setting. This will let users specify a default set of Input mapping contexts without the trouble of getting the player controller in blueprints or c++

diff --git a/Engine/Plugins/EnhancedInput/Source/EnhancedInput/Private/EnhancedPlayerInput.cpp b/Engine/Plugins/EnhancedInput/Source/EnhancedInput/Private/EnhancedPlayerInput.cpp
--- a/Engine/Plugins/EnhancedInput/Source/EnhancedInput/Private/EnhancedPlayerInput.cpp
+++ b/Engine/Plugins/EnhancedInput/Source/EnhancedInput/Private/EnhancedPlayerInput.cpp
@@ -9,6 +9,7 @@
 #include "Engine/World.h"
 #include "GameFramework/WorldSettings.h"
 #include "EnhancedInputDeveloperSettings.h"
+#include "InputMappingContext.h"
 
 #include UE_INLINE_GENERATED_CPP_BY_NAME(EnhancedPlayerInput)
 
@@ -20,6 +21,26 @@ namespace UE
 		static FAutoConsoleVariableRef CVarShouldOnlyTriggerLastActionInChord(TEXT("EnhancedInput.OnlyTriggerLastActionInChord"),
 			ShouldOnlyTriggerLastActionInChord,
 			TEXT("Should only the last action in a ChordedAction trigger be fired? If this is disabled, then the dependant chords will be fired as well"));
+
+		static int32 EnableDefaultMappingContexts = 1;
+		static FAutoConsoleVariableRef CVarEnableDefaultMappingContexts(TEXT("EnhancedInput.EnableDefaultMappingContexts"),
+			EnableDefaultMappingContexts,
+			TEXT("Should the UEnhancedInputDeveloperSettings::DefaultMappingContexts be applied to every UEnhancedPlayerInput?"));
+	}
+}
+
+UEnhancedPlayerInput::UEnhancedPlayerInput()
+	: Super()
+{
+	if (UE::Input::EnableDefaultMappingContexts)
+	{
+		for (const FDefaultContextSetting& DefaultContext : GetDefault<UEnhancedInputDeveloperSettings>()->DefaultMappingContexts)
+		{
+			if (const UInputMappingContext* IMC = DefaultContext.InputMappingContext.LoadSynchronous())
+			{
+				AppliedInputContexts.Add(IMC, DefaultContext.Priority);
+			}
+		}
 	}
 }
 
diff --git a/Engine/Plugins/EnhancedInput/Source/EnhancedInput/Public/EnhancedInputDeveloperSettings.h b/Engine/Plugins/EnhancedInput/Source/EnhancedInput/Public/EnhancedInputDeveloperSettings.h
--- a/Engine/Plugins/EnhancedInput/Source/EnhancedInput/Public/EnhancedInputDeveloperSettings.h
+++ b/Engine/Plugins/EnhancedInput/Source/EnhancedInput/Public/EnhancedInputDeveloperSettings.h
@@ -7,6 +7,23 @@
 #include "Engine/PlatformSettings.h"
 #include "EnhancedInputDeveloperSettings.generated.h"
 
+class UInputMappingContext;
+
+/** Represents a single input mapping context and the priority that it should be applied with */
+USTRUCT()
+struct FDefaultContextSetting
+{
+	GENERATED_BODY()
+
+	/** Input Mapping Context that should be Added to the EnhancedInputEditorSubsystem when it starts listening for input */
+	UPROPERTY(EditAnywhere, Config, Category = "Input")
+	TSoftObjectPtr<const UInputMappingContext> InputMappingContext = nullptr;
+
+	/** The prioirty that should be given to this mapping context when it is added */
+	UPROPERTY(EditAnywhere, Config, Category = "Input")
+	int32 Priority = 0;
+};
+
 /** Developer settings for Enhanced Input */
 UCLASS(config = Input, defaultconfig, meta = (DisplayName = "Enhanced Input"))
 class ENHANCEDINPUT_API UEnhancedInputDeveloperSettings : public UDeveloperSettingsBackedByCVars
@@ -16,6 +33,17 @@ public:
 	
 	UEnhancedInputDeveloperSettings(const FObjectInitializer& Initializer);
 
+	/** If true, then the DefaultMappingContexts will be applied to all Enhanced Player Inputs. */
+	UPROPERTY(EditAnywhere, Category = "Enhanced Input", meta=(ConsoleVariable="EnhancedInput.EnableDefaultMappingContexts"))
+	bool bEnableDefaultMappingContexts = true;
+	
+	/**
+	 * Array of any input mapping contexts that you want to be applied by default to the Enhanced Input local player subsystem.
+	 * NOTE: These mapping context's can only be from your game's root content directory, not plugins.
+	 */
+	UPROPERTY(config, EditAnywhere, Category = "Enhanced Input",  meta = (editCondition = "bEnableDefaultMappingContexts"))
+	TArray<FDefaultContextSetting> DefaultMappingContexts;
+
 	/**
 	 * Platform specific settings for Enhanced Input.
 	 * @see UEnhancedInputPlatformSettings
diff --git a/Engine/Plugins/EnhancedInput/Source/EnhancedInput/Public/EnhancedPlayerInput.h b/Engine/Plugins/EnhancedInput/Source/EnhancedInput/Public/EnhancedPlayerInput.h
--- a/Engine/Plugins/EnhancedInput/Source/EnhancedInput/Public/EnhancedPlayerInput.h
+++ b/Engine/Plugins/EnhancedInput/Source/EnhancedInput/Public/EnhancedPlayerInput.h
@@ -44,6 +44,8 @@ class ENHANCEDINPUT_API UEnhancedPlayerInput : public UPlayerInput
 
 public:
 
+	UEnhancedPlayerInput();
+
 	/**
 	* Returns the action instance data for the given input action if there is any. Returns nullptr if the action is not available.
 	*/
diff --git a/Engine/Plugins/EnhancedInput/Source/InputEditor/Public/EnhancedInputEditorSettings.h b/Engine/Plugins/EnhancedInput/Source/InputEditor/Public/EnhancedInputEditorSettings.h
--- a/Engine/Plugins/EnhancedInput/Source/InputEditor/Public/EnhancedInputEditorSettings.h
+++ b/Engine/Plugins/EnhancedInput/Source/InputEditor/Public/EnhancedInputEditorSettings.h
@@ -5,27 +5,13 @@
 #include "CoreMinimal.h"
 #include "Engine/DeveloperSettings.h"
 #include "Engine/DeveloperSettingsBackedByCVars.h"
+#include "EnhancedInputDeveloperSettings.h"
 
 #include "EnhancedInputEditorSettings.generated.h"
 
 class UEnhancedPlayerInput;
-class UInputMappingContext;
 class UEnhancedInputEditorSubsystem;
 
-USTRUCT()
-struct FDefaultContextSetting
-{
-	GENERATED_BODY()
-
-	/** Input Mapping Context that should be Added to the EnhancedInputEditorSubsystem when it starts listening for input */
-	UPROPERTY(EditAnywhere, Category = "Input")
-	TSoftObjectPtr<const UInputMappingContext> InputMappingContext = nullptr;
-
-	/** The prioirty that should be given to this mapping context when it is added */
-	UPROPERTY(EditAnywhere, Category = "Input")
-	int32 Priority = 0;
-};
-
 /** Settings for the Enhanced Input Editor Subsystem that are persistent between a project's users */
 UCLASS(config = Input, defaultconfig, meta = (DisplayName = "Enhanced Input (Editor Only)"))
 class UEnhancedInputEditorProjectSettings : public UDeveloperSettings
