- Created instanced mappable key settings that can be overridden to add per project functionality.
- Instanced mappable key settings are fetched based on a setting behavior enum.
    - InheritSettingsFromAction to use the setting defined by the action
    - OverrideSettings to override the settings defined in the action, and use the ones defined in the action key mapping.
    - IgnoreSettings to ignore settings completely for this action key mapping.
- Added blueprint facing helper functions to accessing mappable key settings.
- Shifted over data validation in to the Action Key Mapping.

diff --git a/Engine/Plugins/EnhancedInput/Source/EnhancedInput/Private/EnhancedActionKeyMapping.cpp b/Engine/Plugins/EnhancedInput/Source/EnhancedInput/Private/EnhancedActionKeyMapping.cpp
new file mode 100644
--- /dev/null
+++ b/Engine/Plugins/EnhancedInput/Source/EnhancedInput/Private/EnhancedActionKeyMapping.cpp
@@ -0,0 +1,99 @@
+// Copyright Epic Games, Inc. All Rights Reserved.
+
+#include "EnhancedActionKeyMapping.h"
+#include "PlayerMappableKeySettings.h"
+#include "InputTriggers.h"
+#include "InputModifiers.h"
+
+#define LOCTEXT_NAMESPACE "ActionKeyMapping"
+
+UPlayerMappableKeySettings* FEnhancedActionKeyMapping::GetPlayerMappableKeySettings() const
+{
+	switch (SettingBehavior)
+	{
+		case EPlayerMappableKeySettingBehaviors::InheritSettingsFromAction:	
+			return Action != nullptr ? Action->GetPlayerMappableKeySettings() : nullptr;
+		case EPlayerMappableKeySettingBehaviors::OverrideSettings:			
+			return PlayerMappableKeySettings;
+		case EPlayerMappableKeySettingBehaviors::IgnoreSettings:				
+			return nullptr;
+	}
+	return nullptr;
+}
+
+ENHANCEDINPUT_API FName FEnhancedActionKeyMapping::GetMappingName() const
+{
+	if (IsPlayerMappable())
+	{
+		if (UPlayerMappableKeySettings* MappableKeySettings = GetPlayerMappableKeySettings())
+		{
+			return MappableKeySettings->GetMappingName();
+		}
+		return PlayerMappableOptions.Name;
+	}
+	return NAME_None;
+}
+
+bool FEnhancedActionKeyMapping::IsPlayerMappable() const
+{
+	return GetPlayerMappableKeySettings() != nullptr || bIsPlayerMappable;
+}
+
+#if WITH_EDITOR
+EDataValidationResult FEnhancedActionKeyMapping::IsDataValid(TArray<FText>& ValidationErrors)
+{
+	EDataValidationResult Result = EDataValidationResult::Valid;
+
+	//Validate Action Reference.
+	if (Action == nullptr)
+	{
+		Result = EDataValidationResult::Invalid;
+		ValidationErrors.Add(LOCTEXT("NullInputAction", "A mapping cannot have an empty input action!"));
+	}
+
+	//Validate Player Mappable Options Name.
+	if (bIsPlayerMappable && PlayerMappableOptions.Name == NAME_None)
+	{
+		ValidationErrors.Add(LOCTEXT("InvalidPlayerMappableName", "A player mappable key mapping must have a valid 'Name'"));
+		return EDataValidationResult::Invalid;
+	}
+
+	// Validate Settings.
+	if (PlayerMappableKeySettings != nullptr)
+	{
+		Result = CombineDataValidationResults(Result, PlayerMappableKeySettings->IsDataValid(ValidationErrors));
+	}
+
+	// Validate the triggers.
+	for (const TObjectPtr<UInputTrigger> Trigger : Triggers)
+	{
+		if (Trigger != nullptr)
+		{
+			Result = CombineDataValidationResults(Result, Trigger->IsDataValid(ValidationErrors));
+		}
+		else
+		{
+			Result = EDataValidationResult::Invalid;
+			ValidationErrors.Add(LOCTEXT("NullInputTrigger", "There cannot be a null Input Trigger on a key mapping"));
+		}
+	}
+
+	// Validate the modifiers.
+	for (const TObjectPtr<UInputModifier> Modifier : Modifiers)
+	{
+		if (Modifier != nullptr)
+		{
+			Result = CombineDataValidationResults(Result, Modifier->IsDataValid(ValidationErrors));
+		}
+		else
+		{
+			Result = EDataValidationResult::Invalid;
+			ValidationErrors.Add(LOCTEXT("NullInputModifier", "There cannot be a null Input Modifier on a key mapping"));
+		}
+	}
+	return Result;
+}
+
+#endif // WITH_EDITOR
+
+#undef LOCTEXT_NAMESPACE
diff --git a/Engine/Plugins/EnhancedInput/Source/EnhancedInput/Private/EnhancedInputLibrary.cpp b/Engine/Plugins/EnhancedInput/Source/EnhancedInput/Private/EnhancedInputLibrary.cpp
--- a/Engine/Plugins/EnhancedInput/Source/EnhancedInput/Private/EnhancedInputLibrary.cpp
+++ b/Engine/Plugins/EnhancedInput/Source/EnhancedInput/Private/EnhancedInputLibrary.cpp
@@ -62,6 +62,22 @@ FInputActionValue UEnhancedInputLibrary::MakeInputActionValue(double X, double Y
 {
 	return FInputActionValue(MatchValueType.GetValueType(), FVector(X, Y, Z));
 }
+
+UPlayerMappableKeySettings* UEnhancedInputLibrary::GetPlayerMappableKeySettings(const FEnhancedActionKeyMapping& ActionKeyMapping)
+{
+	return ActionKeyMapping.GetPlayerMappableKeySettings();
+}
+
+FName UEnhancedInputLibrary::GetMappingName(const FEnhancedActionKeyMapping& ActionKeyMapping)
+{
+	return ActionKeyMapping.GetMappingName();
+}
+
+bool UEnhancedInputLibrary::IsActionKeyMappingPlayerMappable(const FEnhancedActionKeyMapping& ActionKeyMapping)
+{
+	return ActionKeyMapping.IsPlayerMappable();
+}
+
 PRAGMA_ENABLE_DEPRECATION_WARNINGS
 
 FInputActionValue UEnhancedInputLibrary::MakeInputActionValueOfType(double X, double Y, double Z, const EInputActionValueType ValueType)
diff --git a/Engine/Plugins/EnhancedInput/Source/EnhancedInput/Private/EnhancedInputSubsystemInterface.cpp b/Engine/Plugins/EnhancedInput/Source/EnhancedInput/Private/EnhancedInputSubsystemInterface.cpp
--- a/Engine/Plugins/EnhancedInput/Source/EnhancedInput/Private/EnhancedInputSubsystemInterface.cpp
+++ b/Engine/Plugins/EnhancedInput/Source/EnhancedInput/Private/EnhancedInputSubsystemInterface.cpp
@@ -4,15 +4,16 @@
 
 #include "EnhancedInputComponent.h"
 #include "EnhancedInputModule.h"
+#include "EnhancedInputPlatformSettings.h"
 #include "EnhancedPlayerInput.h"
 #include "GameFramework/PlayerController.h"
 #include "InputCoreTypes.h"
 #include "InputMappingContext.h"
 #include "InputModifiers.h"
 #include "InputTriggers.h"
-#include "UObject/UObjectIterator.h"
 #include "PlayerMappableInputConfig.h"
-#include "EnhancedInputPlatformSettings.h"
+#include "PlayerMappableKeySettings.h"
+#include "UObject/UObjectIterator.h"
 
 #include UE_INLINE_GENERATED_CPP_BY_NAME(EnhancedInputSubsystemInterface)
 
@@ -393,7 +394,7 @@ TArray<FEnhancedActionKeyMapping> IEnhancedInputSubsystemInterface::GetAllPlayer
     {
     	for (const FEnhancedActionKeyMapping& Mapping : PlayerInput->EnhancedActionMappings)
 		{
-			if (Mapping.bIsPlayerMappable)
+			if (Mapping.IsPlayerMappable())
 			{
 				PlayerMappableMappings.AddUnique(Mapping);
 			}
@@ -655,7 +656,7 @@ void IEnhancedInputSubsystemInterface::RebuildControlMappings()
 
 		for (FEnhancedActionKeyMapping& Mapping : OrderedMappings)
 		{
-			if (FKey* PlayerKey = PlayerMappedSettings.Find(Mapping.PlayerMappableOptions.Name))
+			if (FKey* PlayerKey = PlayerMappedSettings.Find(Mapping.GetMappingName()))
 			{
 				Mapping.Key = *PlayerKey;
 			}
diff --git a/Engine/Plugins/EnhancedInput/Source/EnhancedInput/Private/InputAction.cpp b/Engine/Plugins/EnhancedInput/Source/EnhancedInput/Private/InputAction.cpp
--- a/Engine/Plugins/EnhancedInput/Source/EnhancedInput/Private/InputAction.cpp
+++ b/Engine/Plugins/EnhancedInput/Source/EnhancedInput/Private/InputAction.cpp
@@ -5,6 +5,7 @@
 #include "EnhancedPlayerInput.h"
 #include "InputModifiers.h"
 #include "InputTriggers.h"
+#include "PlayerMappableKeySettings.h"
 #include "UObject/UObjectIterator.h"
 
 #include UE_INLINE_GENERATED_CPP_BY_NAME(InputAction)
@@ -148,6 +149,12 @@ EDataValidationResult UInputAction::IsDataValid(TArray<FText>& ValidationErrors)
 			ValidationErrors.Add(LOCTEXT("NullInputModifier", "There cannot be a null Input Modifier on an Input Action!"));
 		}		
 	}
+	
+	// Validate Settings
+	if (PlayerMappableKeySettings != nullptr)
+	{
+		Result = CombineDataValidationResults(Result, PlayerMappableKeySettings->IsDataValid(ValidationErrors));
+	}
 
 	return Result;
 }
@@ -179,6 +186,6 @@ ETriggerEventsSupported UInputAction::GetSupportedTriggerEvents() const
 	return SupportedTriggers;
 }
 
-#endif
+#endif // WITH_EDITOR
 
 #undef LOCTEXT_NAMESPACE
diff --git a/Engine/Plugins/EnhancedInput/Source/EnhancedInput/Private/InputMappingContext.cpp b/Engine/Plugins/EnhancedInput/Source/EnhancedInput/Private/InputMappingContext.cpp
--- a/Engine/Plugins/EnhancedInput/Source/EnhancedInput/Private/InputMappingContext.cpp
+++ b/Engine/Plugins/EnhancedInput/Source/EnhancedInput/Private/InputMappingContext.cpp
@@ -13,50 +13,10 @@
 EDataValidationResult UInputMappingContext::IsDataValid(TArray<FText>& ValidationErrors)
 {
 	EDataValidationResult Result = CombineDataValidationResults(Super::IsDataValid(ValidationErrors), EDataValidationResult::Valid);
-	
-	for (const FEnhancedActionKeyMapping& Mapping : Mappings)
+	for (FEnhancedActionKeyMapping& Mapping : Mappings)
 	{
-		if (!Mapping.Action)
-		{
-			Result = EDataValidationResult::Invalid;
-			ValidationErrors.Add(LOCTEXT("NullInputAction", "A mapping cannot have an empty input action!"));
-		}
-		
-		if (Mapping.bIsPlayerMappable && Mapping.PlayerMappableOptions.Name == NAME_None)
-		{
-			Result = EDataValidationResult::Invalid;
-        	ValidationErrors.Add(LOCTEXT("InvalidPlayerMappableName", "A player mappable key mapping must have a valid 'Name'"));
-		}
-		
-		// Validate the triggers
-		for (const TObjectPtr<UInputTrigger> Trigger : Mapping.Triggers)
-		{
-			if (Trigger)
-			{
-				Result = CombineDataValidationResults(Result, Trigger->IsDataValid(ValidationErrors));
-			}
-			else
-			{
-				Result = EDataValidationResult::Invalid;
-				ValidationErrors.Add(LOCTEXT("NullInputTrigger", "There cannot be a null Input Trigger on a key mapping"));
-			}
-		}
-		
-		// Validate the modifiers
-		for (const TObjectPtr<UInputModifier> Modifier : Mapping.Modifiers)
-		{
-			if (Modifier)
-			{
-				Result = CombineDataValidationResults(Result, Modifier->IsDataValid(ValidationErrors));	
-			}
-			else
-			{
-				Result = EDataValidationResult::Invalid;
-				ValidationErrors.Add(LOCTEXT("NullInputModifier", "There cannot be a null Input Modifier on a key mapping"));
-			}
-		}
+		Result = CombineDataValidationResults(Result, Mapping.IsDataValid(ValidationErrors));
 	}
-
 	return Result;
 }
 #endif	// WITH_EDITOR
diff --git a/Engine/Plugins/EnhancedInput/Source/EnhancedInput/Private/PlayerMappableInputConfig.cpp b/Engine/Plugins/EnhancedInput/Source/EnhancedInput/Private/PlayerMappableInputConfig.cpp
--- a/Engine/Plugins/EnhancedInput/Source/EnhancedInput/Private/PlayerMappableInputConfig.cpp
+++ b/Engine/Plugins/EnhancedInput/Source/EnhancedInput/Private/PlayerMappableInputConfig.cpp
@@ -2,6 +2,7 @@
 
 #include "PlayerMappableInputConfig.h"
 #include "InputMappingContext.h"
+#include "PlayerMappableKeySettings.h"
 
 #include UE_INLINE_GENERATED_CPP_BY_NAME(PlayerMappableInputConfig)
 
@@ -34,22 +35,17 @@ EDataValidationResult UPlayerMappableInputConfig::IsDataValid(TArray<FText>& Val
 	}
 	
 	// Check that every context is valid and all player mappable mappings have names
-	for (const TPair<TObjectPtr<UInputMappingContext>, int32>& ContextPair : Contexts)
+	for (TPair<TObjectPtr<UInputMappingContext>, int32>& ContextPair : Contexts)
 	{
-		if(const UInputMappingContext* IMC = ContextPair.Key.Get())
+		if (UInputMappingContext* IMC = ContextPair.Key.Get())
 		{
-			for(const FEnhancedActionKeyMapping& Mapping : IMC->GetMappings())
+			Result = CombineDataValidationResults(Result, IMC->IsDataValid(ValidationErrors));
+			if (Result == EDataValidationResult::Invalid)
 			{
-				if(Mapping.bIsPlayerMappable && Mapping.PlayerMappableOptions.Name == NAME_None)
-				{
-					Result = CombineDataValidationResults(Result, EDataValidationResult::Invalid);
-					FFormatNamedArguments Args;
-					Args.Add(TEXT("AssetPath"), FText::FromString(GetPathName()));
-					Args.Add(TEXT("ActionName"), FText::FromString(Mapping.Action->GetPathName()));
-
-					const FText NoMappingError = FText::Format(LOCTEXT("NoNamePlayerMappingError", "'{AssetPath}' has a player mappable key that has no name!"), Args);
-					ValidationErrors.Emplace(NoMappingError);
-				}
+				FFormatNamedArguments Args;
+				Args.Add(TEXT("AssetPath"), FText::FromString(IMC->GetPathName()));
+				const FText NoMappingError = FText::Format(LOCTEXT("NoNamePlayerMappingError", "'{AssetPath}' has an invalid mapping."), Args);
+				ValidationErrors.Emplace(NoMappingError);
 			}
 		}
 		else
@@ -77,7 +73,7 @@ void UPlayerMappableInputConfig::ForEachDefaultPlayerMappableKey(TFunctionRef<vo
 		{
 			for(const FEnhancedActionKeyMapping& Mapping : IMC->GetMappings())
 			{
-				if(Mapping.bIsPlayerMappable)
+				if(Mapping.IsPlayerMappable())
 				{
 					Operation(Mapping);
 				}
@@ -104,7 +100,7 @@ FEnhancedActionKeyMapping UPlayerMappableInputConfig::GetMappingByName(const FNa
 	FEnhancedActionKeyMapping OutMapping;
 	
 	TArray<FEnhancedActionKeyMapping> MappableKeys = GetPlayerMappableKeys();
-	FEnhancedActionKeyMapping* ExistingMapping = MappableKeys.FindByPredicate([MappingName](const FEnhancedActionKeyMapping& Mapping) { return Mapping.PlayerMappableOptions.Name == MappingName; });
+	FEnhancedActionKeyMapping* ExistingMapping = MappableKeys.FindByPredicate([MappingName](const FEnhancedActionKeyMapping& Mapping) { return Mapping.GetMappingName() == MappingName; });
 
 	if(ExistingMapping)
 	{
diff --git a/Engine/Plugins/EnhancedInput/Source/EnhancedInput/Private/PlayerMappableKeySettings.cpp b/Engine/Plugins/EnhancedInput/Source/EnhancedInput/Private/PlayerMappableKeySettings.cpp
new file mode 100644
--- /dev/null
+++ b/Engine/Plugins/EnhancedInput/Source/EnhancedInput/Private/PlayerMappableKeySettings.cpp
@@ -0,0 +1,20 @@
+// Copyright Epic Games, Inc. All Rights Reserved.
+
+#include "PlayerMappableKeySettings.h"
+
+#define LOCTEXT_NAMESPACE "EnhancedActionKeySetting"
+
+#if WITH_EDITOR
+EDataValidationResult UPlayerMappableKeySettings::IsDataValid(TArray<FText>& ValidationErrors)
+{
+	EDataValidationResult Result = CombineDataValidationResults(Super::IsDataValid(ValidationErrors), EDataValidationResult::Valid);
+	if (Name == NAME_None)
+	{
+		Result = EDataValidationResult::Invalid;
+		ValidationErrors.Add(LOCTEXT("InvalidPlayerMappableKeySettingsName", "A Player Mappable Key Settings must have a valid 'Name'"));
+	}
+	return Result;
+}
+#endif // WITH_EDITOR
+
+#undef LOCTEXT_NAMESPACE
diff --git a/Engine/Plugins/EnhancedInput/Source/EnhancedInput/Public/EnhancedActionKeyMapping.h b/Engine/Plugins/EnhancedInput/Source/EnhancedInput/Public/EnhancedActionKeyMapping.h
--- a/Engine/Plugins/EnhancedInput/Source/EnhancedInput/Public/EnhancedActionKeyMapping.h
+++ b/Engine/Plugins/EnhancedInput/Source/EnhancedInput/Public/EnhancedActionKeyMapping.h
@@ -2,13 +2,36 @@
 
 #pragma once
 
-#include "CoreMinimal.h"
+#include "Containers/Array.h"
 #include "InputAction.h"
 #include "InputActionValue.h"
 #include "InputCoreTypes.h"
+#include "Internationalization/Text.h"
+#include "UObject/NameTypes.h"
+#include "UObject/ObjectPtr.h"
 
 #include "EnhancedActionKeyMapping.generated.h"
 
+class UInputModifier;
+class UInputTrigger;
+class UPlayerMappableKeySettings;
+
+enum class EDataValidationResult : uint8;
+
+/**
+* Defines which Player Mappable Key Setting to use for a Action Key Mapping.
+*/
+UENUM(BlueprintType)
+enum class EPlayerMappableKeySettingBehaviors : uint8
+{
+	//Use the Settings specified in the Input Action.
+	InheritSettingsFromAction,
+	//Use the Settings specified in the Action Key Mapping overriding the ones specified in the Input action.
+	OverrideSettings,
+	//Don't use any Settings even if one is specified in the Input Action.
+	IgnoreSettings
+};
+
 /**
  * A struct that represents player facing mapping options for an action key mapping.
  * Use this to set a unique FName for the mapping option to save it, as well as some FText options
@@ -40,7 +63,7 @@ public:
 	UPROPERTY(EditAnywhere, BlueprintReadOnly, Category = "Input|PlayerMappable")
 	TObjectPtr<UObject> Metadata = nullptr;
 	
-	/** A unique name for this player binding to be saved with. */
+	/** A unique name for this player mapping to be saved with. */
 	UPROPERTY(EditAnywhere, BlueprintReadOnly, Category = "Input|PlayerMappable")
 	FName Name;
 	
@@ -48,14 +71,11 @@ public:
 	UPROPERTY(EditAnywhere, BlueprintReadOnly, Category = "Input|PlayerMappable")
 	FText DisplayName;
 
-	/** The category that this player binding is in */
+	/** The category that this player mapping is in */
 	UPROPERTY(EditAnywhere, BlueprintReadOnly, Category = "Input|PlayerMappable")
 	FText DisplayCategory = FText::GetEmpty();
 };
 
-class UInputTrigger;
-class UInputModifier;
-
 /**
  * Defines a mapping between a key activation and the resulting enhanced action
  * An key could be a button press, joystick axis movement, etc.
@@ -67,10 +87,29 @@ struct FEnhancedActionKeyMapping
 {
 	GENERATED_BODY()
 
+	/**
+	* Returns the Player Mappable Key Settings owned by the Action Key Mapping or by the referenced Input Action, or nothing based of the Setting Behavior.
+	*/
+	ENHANCEDINPUT_API UPlayerMappableKeySettings* GetPlayerMappableKeySettings() const;
+
+	/**
+	 * Returns the name of the mapping based on setting behavior used. If no name is found in the Mappable Key Settings it will return the name set in Player Mappable Options if bIsPlayerMappable is true.
+	 */
+	ENHANCEDINPUT_API FName GetMappingName() const;
+
+	/**
+	* Returns true if this Action Key Mapping either holds a Player Mappable Key Settings or is set bIsPlayerMappable.
+	*/
+	ENHANCEDINPUT_API bool IsPlayerMappable() const;
+
+#if WITH_EDITOR
+	EDataValidationResult IsDataValid(TArray<FText>& ValidationErrors);
+#endif
+
 	/** Options for making this a player mappable keymapping */
 	UPROPERTY(EditAnywhere, BlueprintReadOnly, Category = "Input|PlayerMappable", meta = (editCondition = "bIsPlayerMappable", DisplayAfter = "bIsPlayerMappable"))
 	FPlayerMappableKeyOptions PlayerMappableOptions;
-	
+
 	/**
 	* Trigger qualifiers. If any trigger qualifiers exist the mapping will not trigger unless:
 	* If there are any Explicit triggers in this list at least one of them must be met.
@@ -94,10 +133,6 @@ struct FEnhancedActionKeyMapping
 	UPROPERTY(EditAnywhere, BlueprintReadOnly, Category = "Input")
 	FKey Key;
 
-	/** If true than this ActionKeyMapping will be exposed as a player mappable key */
-	UPROPERTY(EditAnywhere, BlueprintReadOnly, Category = "Input|PlayerMappable")
-	uint8 bIsPlayerMappable : 1;
-
 	/**
 	 * If true, then this Key Mapping should be ignored. This is set to true if the key is down
 	 * during a rebuild of it's owning PlayerInput ControlMappings.
@@ -119,8 +154,36 @@ struct FEnhancedActionKeyMapping
 		: PlayerMappableOptions(InAction)
 		, Action(InAction)
 		, Key(InKey)
-		, bIsPlayerMappable(false)
 		, bShouldBeIgnored(false)
+		, bIsPlayerMappable(false)
 	{}
 
+	friend class FEnhancedActionMappingCustomization;
+
+protected:
+
+	/**
+	* Defines which Player Mappable Key Setting to use for a Action Key Mapping.
+	*/
+	UPROPERTY(EditAnywhere, Category = "Input|Settings")
+	EPlayerMappableKeySettingBehaviors SettingBehavior = EPlayerMappableKeySettingBehaviors::InheritSettingsFromAction;
+
+	/**
+	* Used to expose this mapping or to opt-out of settings completely.
+	*/
+	UPROPERTY(EditAnywhere, Instanced, Category = "Input|Settings", meta = (EditCondition = "SettingBehavior == EPlayerMappableKeySettingBehaviors::OverrideSettings", DisplayAfter = "SettingBehavior"))
+	TObjectPtr<UPlayerMappableKeySettings> PlayerMappableKeySettings = nullptr;
+
+	/** If true than this ActionKeyMapping will be exposed as a player mappable key */
+	UPROPERTY(EditAnywhere, BlueprintReadOnly, Category = "Input|PlayerMappable")
+	uint8 bIsPlayerMappable : 1;
+
+public:
+
+	template<typename T = UPlayerMappableKeySettings> 
+	T* GetPlayerMappableKeySettings() const
+	{
+		return Cast<T>(GetPlayerMappableKeySettings());
+	}
+
 };
diff --git a/Engine/Plugins/EnhancedInput/Source/EnhancedInput/Public/EnhancedInputLibrary.h b/Engine/Plugins/EnhancedInput/Source/EnhancedInput/Public/EnhancedInputLibrary.h
--- a/Engine/Plugins/EnhancedInput/Source/EnhancedInput/Public/EnhancedInputLibrary.h
+++ b/Engine/Plugins/EnhancedInput/Source/EnhancedInput/Public/EnhancedInputLibrary.h
@@ -12,9 +12,10 @@
 #include "EnhancedInputLibrary.generated.h"
 
 class APlayerController;
-class UInputMappingContext;
-class UInputAction;
 class UEnhancedPlayerInput;
+class UInputAction;
+class UInputMappingContext;
+class UPlayerMappableKeySettings;
 
 UCLASS()
 class ENHANCEDINPUT_API UEnhancedInputLibrary : public UBlueprintFunctionLibrary
@@ -51,6 +52,24 @@ public:
 	UFUNCTION(BlueprintPure, Category = "Input", meta = (Keywords = "construct build", NativeMakeFunc, DeprecatedFunction, DeprecatedMessage="This version of MakeInputActionValue has been deprecated, please use MakeInputActionValueOfType"))
 	static FInputActionValue MakeInputActionValue(double X, double Y, double Z, const FInputActionValue& MatchValueType);
 
+	/**
+	* Returns the Player Mappable Key Settings owned by the Action Key Mapping or by the referenced Input Action, or nothing based of the Setting Behavior.
+	*/
+	UFUNCTION(BlueprintPure, Category = "Input", meta = (ReturnDisplayName = "Player Mappable Key Settings"))
+	static UPlayerMappableKeySettings* GetPlayerMappableKeySettings(UPARAM(ref) const FEnhancedActionKeyMapping& ActionKeyMapping);
+
+	/**
+	* Returns the name of the mapping based on setting behavior used. If no name is found in the Mappable Key Settings it will return the name set in Player Mappable Options if bIsPlayerMappable is true.
+	*/
+	UFUNCTION(BlueprintPure, Category = "Input", meta = (ReturnDisplayName = "Mapping Name"))
+	static FName GetMappingName(UPARAM(ref) const FEnhancedActionKeyMapping& ActionKeyMapping);
+
+	/**
+	 * Returns true if this Action Key Mapping either holds a Player Mappable Key Settings or is set bIsPlayerMappable.
+	 */
+	UFUNCTION(BlueprintPure, Category = "Input", meta = (ReturnDisplayName = "Is Player Mappable"))
+	static bool IsActionKeyMappingPlayerMappable(UPARAM(ref) const FEnhancedActionKeyMapping& ActionKeyMapping);
+
 	// Internal helper functionality
 
 	// GetInputActionvalue internal accessor function for actions that have been bound to from a UEnhancedInputComponent
diff --git a/Engine/Plugins/EnhancedInput/Source/EnhancedInput/Public/InputAction.h b/Engine/Plugins/EnhancedInput/Source/EnhancedInput/Public/InputAction.h
--- a/Engine/Plugins/EnhancedInput/Source/EnhancedInput/Public/InputAction.h
+++ b/Engine/Plugins/EnhancedInput/Source/EnhancedInput/Public/InputAction.h
@@ -11,6 +11,8 @@
 
 #include "InputAction.generated.h"
 
+class UPlayerMappableKeySettings;
+
 enum class ETriggerEventInternal : uint8;
 
 // Input action definition. These are instanced per player (via FInputActionInstance)
@@ -32,8 +34,13 @@ public:
 	 * Returns a bitmask of supported trigger events that is built from each UInputTrigger on this Action.
 	 */
 	ETriggerEventsSupported GetSupportedTriggerEvents() const;
-#endif
+#endif // WITH_EDITOR
 	
+	/**
+	* Returns the Player Mappable Key Settings for this Input Action.
+	*/
+	const TObjectPtr<UPlayerMappableKeySettings>& GetPlayerMappableKeySettings() const { return PlayerMappableKeySettings; }
+
 	// A localized descriptor of this input action
 	UPROPERTY(EditAnywhere, BlueprintReadOnly, Category = "Description")
 	FText ActionDescription = FText::GetEmpty();
@@ -69,6 +76,14 @@ public:
 	*/
 	UPROPERTY(EditAnywhere, Instanced, BlueprintReadWrite, Category = Action, meta=(DisplayAfter="Triggers"))
 	TArray<TObjectPtr<UInputModifier>> Modifiers;
+
+private:
+
+	/**
+	* Holds setting information about this Action Input for setting screen and save purposes.
+	*/
+	UPROPERTY(EditAnywhere, Instanced, BlueprintReadWrite, Category = "Input|Settings", meta=(AllowPrivateAccess))
+	TObjectPtr<UPlayerMappableKeySettings> PlayerMappableKeySettings;
 };
 
 // Calculate a collective representation of trigger state from evaluations of all triggers in one or more trigger groups.
diff --git a/Engine/Plugins/EnhancedInput/Source/EnhancedInput/Public/PlayerMappableKeySettings.h b/Engine/Plugins/EnhancedInput/Source/EnhancedInput/Public/PlayerMappableKeySettings.h
new file mode 100644
--- /dev/null
+++ b/Engine/Plugins/EnhancedInput/Source/EnhancedInput/Public/PlayerMappableKeySettings.h
@@ -0,0 +1,42 @@
+// Copyright Epic Games, Inc. All Rights Reserved.
+
+#pragma once
+
+#include "Internationalization/Text.h"
+#include "UObject/NameTypes.h"
+#include "UObject/Object.h"
+#include "UObject/ObjectPtr.h"
+#include "PlayerMappableKeySettings.generated.h"
+
+/**
+* Hold setting information of an Action Input or a Action Key Mapping for setting screen and save purposes.
+*/
+UCLASS(DefaultToInstanced, EditInlineNew)
+class ENHANCEDINPUT_API UPlayerMappableKeySettings : public UObject
+{
+	GENERATED_BODY()
+
+public:
+
+	virtual FName GetMappingName() const { return Name; }
+
+#if WITH_EDITOR
+	virtual EDataValidationResult IsDataValid(TArray<FText>& ValidationErrors) override;
+#endif // WITH_EDITOR
+
+	/** Metadata that can used to store any other related items to this key mapping such as icons, ability assets, etc. */
+	UPROPERTY(EditAnywhere, BlueprintReadOnly, Category = "Settings")
+	TObjectPtr<UObject> Metadata = nullptr;
+
+	/** A unique name for this player mapping to be saved with. */
+	UPROPERTY(EditAnywhere, BlueprintReadOnly, Category = "Settings")
+	FName Name;
+
+	/** The localized display name of this key mapping. Use this when displaying the mappings to a user. */
+	UPROPERTY(EditAnywhere, BlueprintReadOnly, Category = "Settings")
+	FText DisplayName;
+
+	/** The category that this player mapping is in */
+	UPROPERTY(EditAnywhere, BlueprintReadOnly, Category = "Settings")
+	FText DisplayCategory = FText::GetEmpty();
+};
diff --git a/Engine/Plugins/EnhancedInput/Source/InputEditor/Private/InputCustomizations.cpp b/Engine/Plugins/EnhancedInput/Source/InputEditor/Private/InputCustomizations.cpp
--- a/Engine/Plugins/EnhancedInput/Source/InputEditor/Private/InputCustomizations.cpp
+++ b/Engine/Plugins/EnhancedInput/Source/InputEditor/Private/InputCustomizations.cpp
@@ -71,12 +71,16 @@ void FEnhancedActionMappingCustomization::CustomizeChildren(TSharedRef<IProperty
 	TSharedPtr<IPropertyHandle> ModifiersHandle = PropertyHandle->GetChildHandle(GET_MEMBER_NAME_CHECKED(FEnhancedActionKeyMapping, Modifiers));
 	TSharedPtr<IPropertyHandle> IsPlayerMappableHandle = PropertyHandle->GetChildHandle(GET_MEMBER_NAME_CHECKED(FEnhancedActionKeyMapping, bIsPlayerMappable));
 	TSharedPtr<IPropertyHandle> PlayerBindingOptions = PropertyHandle->GetChildHandle(GET_MEMBER_NAME_CHECKED(FEnhancedActionKeyMapping, PlayerMappableOptions));
+	TSharedPtr<IPropertyHandle> SettingBehavior = PropertyHandle->GetChildHandle(GET_MEMBER_NAME_CHECKED(FEnhancedActionKeyMapping, SettingBehavior));
+	TSharedPtr<IPropertyHandle> PlayerMappableKeySettings = PropertyHandle->GetChildHandle(GET_MEMBER_NAME_CHECKED(FEnhancedActionKeyMapping, PlayerMappableKeySettings));
 
 	// TODO: ResetToDefault needs to be disabled for arrays
 	ChildBuilder.AddProperty(TriggersHandle.ToSharedRef());
 	ChildBuilder.AddProperty(ModifiersHandle.ToSharedRef());
 	ChildBuilder.AddProperty(IsPlayerMappableHandle.ToSharedRef());
 	ChildBuilder.AddProperty(PlayerBindingOptions.ToSharedRef());
+	ChildBuilder.AddProperty(SettingBehavior.ToSharedRef());
+	ChildBuilder.AddProperty(PlayerMappableKeySettings.ToSharedRef());
 }
 
 void FEnhancedActionMappingCustomization::RemoveMappingButton_OnClick() const
diff --git a/Engine/Plugins/Experimental/VirtualCamera/Source/VCamCore/Private/VCamComponent.cpp b/Engine/Plugins/Experimental/VirtualCamera/Source/VCamCore/Private/VCamComponent.cpp
--- a/Engine/Plugins/Experimental/VirtualCamera/Source/VCamCore/Private/VCamComponent.cpp
+++ b/Engine/Plugins/Experimental/VirtualCamera/Source/VCamCore/Private/VCamComponent.cpp
@@ -15,6 +15,7 @@
 #include "EnhancedInputSubsystems.h"
 #include "Engine/InputDelegateBinding.h"
 #include "InputMappingContext.h"
+#include "PlayerMappableKeySettings.h"
 
 #include "GameFramework/InputSettings.h"
 
@@ -581,9 +582,9 @@ bool UVCamComponent::AddInputProfileWithCurrentlyActiveMappings(const FName Prof
 			}
 			for (const FEnhancedActionKeyMapping& Mapping : MappingContext->GetMappings())
 			{
-				if (Mapping.bIsPlayerMappable)
+				if (Mapping.IsPlayerMappable())
 				{
-					const FName MappingName = Mapping.PlayerMappableOptions.Name;
+					const FName MappingName = Mapping.GetMappingName();
 
 					// Prefer to use the current mapped key but fallback to the default if no key is mapped
 					FKey CurrentKey = EnhancedInputSubsystemInterface->GetPlayerMappedKey(MappingName);
