Displaying Input Action Triggers/Modifiers in Input Mapping Contexts to improve clarity

diff --git a/Engine/Plugins/EnhancedInput/Source/InputEditor/Private/InputCustomizations.cpp b/Engine/Plugins/EnhancedInput/Source/InputEditor/Private/InputCustomizations.cpp
--- a/Engine/Plugins/EnhancedInput/Source/InputEditor/Private/InputCustomizations.cpp
+++ b/Engine/Plugins/EnhancedInput/Source/InputEditor/Private/InputCustomizations.cpp
@@ -65,6 +65,44 @@ void FEnhancedActionMappingCustomization::CustomizeHeader(TSharedRef<IPropertyHa
 	KeyStructCustomization->CustomizeHeaderOnlyWithButton(KeyHandle.ToSharedRef(), HeaderRow, CustomizationUtils, RemoveButton);
 }
 
+void AddInputActionProperties(TSharedRef<IPropertyHandle> PropertyHandle, IDetailChildrenBuilder& ChildBuilder)
+{
+	void* Data = nullptr;
+    if (PropertyHandle->GetValueData(Data) != FPropertyAccess::Fail)
+    {
+    	TObjectPtr<FEnhancedActionKeyMapping> ActionKeyMapping = static_cast<FEnhancedActionKeyMapping*>(Data);
+    	
+    	// Make sure ActionKeyMapping action is valid. If triggers and modifiers are empty we can just back out - nothing to add
+    	if (ActionKeyMapping->Action && (!ActionKeyMapping->Action->Triggers.IsEmpty() || !ActionKeyMapping->Action->Modifiers.IsEmpty()))
+    	{
+    		// Convert InputAction to non const so we can properly use it in the UObject Array for AddExternalObjectProperty
+    		const UInputAction* InputActionPtr = ActionKeyMapping->Action;
+    		UInputAction* InputAction = const_cast<UInputAction*>(InputActionPtr);
+    		TArray<UObject*> ActionsAsUObjects { InputAction };
+    		// no need to do it if the trigger array is empty
+    		if (!ActionKeyMapping->Action->Triggers.IsEmpty())
+    		{
+    			if (IDetailPropertyRow* InputActionTriggersRow = ChildBuilder.AddExternalObjectProperty(ActionsAsUObjects, GET_MEMBER_NAME_CHECKED(UInputAction, Triggers)))
+    			{
+    				InputActionTriggersRow->DisplayName(LOCTEXT("InputActionTriggersDisplayName", "Triggers From Input Action"));
+    				InputActionTriggersRow->ToolTip(FText::Format(LOCTEXT("InputActionTriggersToolTip", "Triggers from the {0} Input Action"), FText::FromName(ActionKeyMapping->Action.GetFName())));
+    				InputActionTriggersRow->IsEnabled(false);
+    			}
+    		}
+    		// no need to do it if the modifier array is empty
+    		if (!ActionKeyMapping->Action->Modifiers.IsEmpty())
+    		{
+    			if (IDetailPropertyRow* InputActionModifiersRow = ChildBuilder.AddExternalObjectProperty(ActionsAsUObjects, GET_MEMBER_NAME_CHECKED(UInputAction, Modifiers)))
+    			{
+    				InputActionModifiersRow->DisplayName(LOCTEXT("InputActionModifiersDisplayName", "Modifiers From Input Action"));
+    				InputActionModifiersRow->ToolTip(FText::Format(LOCTEXT("InputActionModifiersToolTip", "Modifiers from the {0} Input Action"), FText::FromName(ActionKeyMapping->Action.GetFName())));
+    				InputActionModifiersRow->IsEnabled(false);
+    			}
+    		}
+    	}
+    }
+}
+
 void FEnhancedActionMappingCustomization::CustomizeChildren(TSharedRef<IPropertyHandle> PropertyHandle, IDetailChildrenBuilder& ChildBuilder, IPropertyTypeCustomizationUtils& CustomizationUtils)
 {
 	TSharedPtr<IPropertyHandle> TriggersHandle = PropertyHandle->GetChildHandle(GET_MEMBER_NAME_CHECKED(FEnhancedActionKeyMapping, Triggers));
@@ -77,6 +115,7 @@ void FEnhancedActionMappingCustomization::CustomizeChildren(TSharedRef<IProperty
 	// TODO: ResetToDefault needs to be disabled for arrays
 	ChildBuilder.AddProperty(TriggersHandle.ToSharedRef());
 	ChildBuilder.AddProperty(ModifiersHandle.ToSharedRef());
+	AddInputActionProperties(PropertyHandle, ChildBuilder);
 	ChildBuilder.AddProperty(IsPlayerMappableHandle.ToSharedRef());
 	ChildBuilder.AddProperty(PlayerBindingOptions.ToSharedRef());
 	ChildBuilder.AddProperty(SettingBehavior.ToSharedRef());
