Making a Drag/Drop action for Input Actions into the Blueprint Graph

diff --git a/Engine/Plugins/EnhancedInput/Source/InputEditor/Private/InputEditorModule.cpp b/Engine/Plugins/EnhancedInput/Source/InputEditor/Private/InputEditorModule.cpp
--- a/Engine/Plugins/EnhancedInput/Source/InputEditor/Private/InputEditorModule.cpp
+++ b/Engine/Plugins/EnhancedInput/Source/InputEditor/Private/InputEditorModule.cpp
@@ -2,7 +2,9 @@
 
 #include "InputEditorModule.h"
 
+#include "AssetBlueprintGraphActions.h"
 #include "AssetTypeActions_Base.h"
+#include "BlueprintEditorModule.h"
 #include "EnhancedInputModule.h"
 #include "DetailCategoryBuilder.h"
 #include "DetailLayoutBuilder.h"
@@ -33,6 +35,7 @@
 #include "Styling/SlateStyleMacros.h"
 #include "Styling/SlateStyleRegistry.h"
 #include "ContentBrowserModule.h"
+#include "EdGraphSchema_K2_Actions.h"
 #include "IContentBrowserSingleton.h"
 #include "ClassViewerModule.h"
 #include "ClassViewerFilter.h"
@@ -44,6 +47,7 @@
 #include "Interfaces/IMainFrameModule.h"
 #include "SourceControlHelpers.h"
 #include "HAL/FileManager.h"
+#include "Kismet2/KismetEditorUtilities.h"
 
 #include UE_INLINE_GENERATED_CPP_BY_NAME(InputEditorModule)
 
@@ -327,6 +331,55 @@ public:
 	virtual UClass* GetSupportedClass() const override { return UPlayerMappableInputConfig::StaticClass(); }
 };
 
+struct FInputActionGraphActions : public FAssetBlueprintGraphActions
+{
+	virtual FText GetGraphHoverMessage(const FAssetData& AssetData, const UEdGraph* HoverGraph) const override;
+	virtual bool TryCreatingAssetNode(const FAssetData& AssetData, UEdGraph* ParentGraph, const FVector2D Location, EK2NewNodeFlags Options) const override;
+};
+
+FText FInputActionGraphActions::GetGraphHoverMessage(const FAssetData& AssetData, const UEdGraph* HoverGraph) const
+{
+	return FText::Format(LOCTEXT("InputActionHoverMessage", "{0}"), FText::FromName(AssetData.AssetName));
+}
+
+bool FInputActionGraphActions::TryCreatingAssetNode(const FAssetData& AssetData, UEdGraph* ParentGraph, const FVector2D Location, EK2NewNodeFlags Options) const
+{
+	if (AssetData.IsValid())
+	{
+		if (const UInputAction* Action = Cast<const UInputAction>(AssetData.GetAsset()))
+		{
+			for (TObjectPtr<UEdGraphNode> Node : ParentGraph->Nodes)
+			{
+				if(const UK2Node_EnhancedInputAction* InputActionNode = Cast<UK2Node_EnhancedInputAction>(Node))
+				{
+					if (InputActionNode->InputAction.GetFName() == AssetData.AssetName)
+					{
+						if (const TSharedPtr<IBlueprintEditor> BlueprintEditor = FKismetEditorUtilities::GetIBlueprintEditorForObject(ParentGraph, false))
+						{
+							BlueprintEditor.Get()->JumpToPin(InputActionNode->GetPinAt(0));
+						}
+						
+						return false;
+					}
+				}
+			}
+
+			UK2Node_EnhancedInputAction* NewNode = FEdGraphSchemaAction_K2NewNode::SpawnNode<UK2Node_EnhancedInputAction>(
+				ParentGraph,
+				Location,
+				Options,
+				[Action](UK2Node_EnhancedInputAction* NewInstance)
+				{
+					NewInstance->InputAction = Action;
+				}
+
+			);
+			return true;
+		}
+	}
+	return false;
+}
+
 /** Custom style set for Enhanced Input */
 class FEnhancedInputSlateStyle final : public FSlateStyleSet
 {
@@ -379,6 +432,12 @@ void FInputEditorModule::StartupModule()
 		//RegisterAssetTypeActions(AssetTools, MakeShareable(new FAssetTypeActions_InputModifier));
 	}
 
+	// Register graph actions:
+	FBlueprintGraphModule& GraphModule = FModuleManager::LoadModuleChecked<FBlueprintGraphModule>("BlueprintGraph");
+	{
+		GraphModule.RegisterGraphAction(UInputAction::StaticClass(), MakeUnique<FInputActionGraphActions>());
+	}
+
 	// Make a new style set for Enhanced Input, which will register any custom icons for the types in this plugin
 	StyleSet = MakeShared<FEnhancedInputSlateStyle>();
 	FSlateStyleRegistry::RegisterSlateStyle(*StyleSet.Get());
diff --git a/Engine/Source/Editor/BlueprintGraph/Classes/AssetBlueprintGraphActions.h b/Engine/Source/Editor/BlueprintGraph/Classes/AssetBlueprintGraphActions.h
new file mode 100644
--- /dev/null
+++ b/Engine/Source/Editor/BlueprintGraph/Classes/AssetBlueprintGraphActions.h
@@ -0,0 +1,36 @@
+// Copyright Epic Games, Inc. All Rights Reserved.
+
+#pragma once
+
+#include "CoreMinimal.h"
+
+struct FAssetData;
+class UEdGraph;
+enum class EK2NewNodeFlags;
+
+/** Struct to handle Asset Blueprint Graph actions
+ * Essentially this is used to help generate asset nodes/responses for Blueprint graphs and keeping Blueprint code from being dependent on something it shouldn't
+ * For example, this can be used to help with dragging/dropping assets on a Blueprint graph
+ */
+struct BLUEPRINTGRAPH_API FAssetBlueprintGraphActions
+{
+	virtual ~FAssetBlueprintGraphActions() {}
+
+	/**
+	 * Gets the hover message when when holding a asset over a Blueprint graph from a drag/drop operation
+	 * @param AssetData The asset being dragged/dropped
+	 * @param HoverGraph The Blueprint graph to drop the asset on
+	 * @return Message to display when hovering.
+	 */
+	virtual FText GetGraphHoverMessage(const FAssetData& AssetData, const UEdGraph* HoverGraph) const = 0;
+
+	/**
+	 * Tries to create a node from the asset passed on the graph passed in
+	 * @param AssetData The asset being dragged/dropped
+	 * @param ParentGraph The Blueprint graph to drop the asset on
+	 * @param Location Location to spawn the new asset node at
+	 * @param Options Flags for node creation
+	 * @return Bool of whether a node was successfully created
+	 */
+	virtual bool TryCreatingAssetNode(const FAssetData& AssetData, UEdGraph* ParentGraph, const FVector2D Location, EK2NewNodeFlags Options) const = 0;
+};
diff --git a/Engine/Source/Editor/BlueprintGraph/Private/AssetBlueprintGraphActions.cpp b/Engine/Source/Editor/BlueprintGraph/Private/AssetBlueprintGraphActions.cpp
new file mode 100644
--- /dev/null
+++ b/Engine/Source/Editor/BlueprintGraph/Private/AssetBlueprintGraphActions.cpp
@@ -0,0 +1,3 @@
+// Copyright Epic Games, Inc. All Rights Reserved.
+
+#include "AssetBlueprintGraphActions.h"
diff --git a/Engine/Source/Editor/BlueprintGraph/Private/BlueprintGraphModule.cpp b/Engine/Source/Editor/BlueprintGraph/Private/BlueprintGraphModule.cpp
--- a/Engine/Source/Editor/BlueprintGraph/Private/BlueprintGraphModule.cpp
+++ b/Engine/Source/Editor/BlueprintGraph/Private/BlueprintGraphModule.cpp
@@ -2,6 +2,7 @@
 
 #include "BlueprintGraphModule.h"
 
+#include "AssetBlueprintGraphActions.h"
 #include "EdGraphSchema_K2.h"
 #include "BlueprintTypePromotion.h"
 #include "Modules/ModuleManager.h"
@@ -12,4 +13,25 @@ void FBlueprintGraphModule::ShutdownModule()
 {
 	UEdGraphSchema_K2::Shutdown();
 	FTypePromotion::Shutdown();
-}
\ No newline at end of file
+	AssetBlueprintGraphActions.Reset();
+}
+
+void FBlueprintGraphModule::RegisterGraphAction(const UClass* AssetType, TUniquePtr<FAssetBlueprintGraphActions> Action)
+{	
+	ensure(!AssetBlueprintGraphActions.Contains(AssetType));
+	AssetBlueprintGraphActions.Add(AssetType, MoveTemp(Action));
+}
+
+void FBlueprintGraphModule::UnregisterGraphAction(const UClass* AssetType)
+{
+	AssetBlueprintGraphActions.Remove(AssetType);
+}
+
+const FAssetBlueprintGraphActions* FBlueprintGraphModule::GetAssetBlueprintGraphActions(const UClass* AssetType) const
+{
+	if (const TUniquePtr<FAssetBlueprintGraphActions>* Actions = AssetBlueprintGraphActions.Find(AssetType))
+	{
+		return Actions->Get();
+	}
+	return nullptr;
+}
diff --git a/Engine/Source/Editor/BlueprintGraph/Private/EdGraphSchema_K2.cpp b/Engine/Source/Editor/BlueprintGraph/Private/EdGraphSchema_K2.cpp
--- a/Engine/Source/Editor/BlueprintGraph/Private/EdGraphSchema_K2.cpp
+++ b/Engine/Source/Editor/BlueprintGraph/Private/EdGraphSchema_K2.cpp
@@ -1,6 +1,7 @@
 // Copyright Epic Games, Inc. All Rights Reserved.
 
 #include "EdGraphSchema_K2.h"
+#include "AssetBlueprintGraphActions.h"
 #include "BlueprintCompilationManager.h"
 #include "Kismet2/Breakpoint.h"
 #include "Modules/ModuleManager.h"
@@ -5765,8 +5766,23 @@ bool UEdGraphSchema_K2::IsStaticFunctionGraph( const UEdGraph* TestEdGraph ) con
 
 void UEdGraphSchema_K2::DroppedAssetsOnGraph(const TArray<FAssetData>& Assets, const FVector2D& GraphPosition, UEdGraph* Graph) const 
 {
+	// only want to spawn event nodes in an event graph
+	if (FBlueprintEditorUtils::IsEventGraph(Graph))
+	{
+		const FBlueprintGraphModule& Module = FModuleManager::LoadModuleChecked<FBlueprintGraphModule>("BlueprintGraph");
+		// check all assets to see if we can get some AssetBlueprintGraphActions for it 
+		for (const FAssetData& AssetData : Assets)
+		{
+			// if we can find any actions we want to try to spawn the node - only spawning Input Action event nodes currently
+			if (const FAssetBlueprintGraphActions* GraphActions = Module.GetAssetBlueprintGraphActions(AssetData.GetClass()))
+			{
+				GraphActions->TryCreatingAssetNode(AssetData, Graph, GraphPosition, EK2NewNodeFlags::SelectNewNode);
+			}
+		}
+	}
+
 	UBlueprint* Blueprint = FBlueprintEditorUtils::FindBlueprintForGraph(Graph);
-	if ((Blueprint != NULL) && FBlueprintEditorUtils::IsActorBased(Blueprint))
+	if ((Blueprint != nullptr) && FBlueprintEditorUtils::IsActorBased(Blueprint))
 	{
 		float XOffset = 0.0f;
 		for(int32 AssetIdx=0; AssetIdx < Assets.Num(); AssetIdx++)
@@ -5878,8 +5894,32 @@ void UEdGraphSchema_K2::GetAssetsGraphHoverMessage(const TArray<FAssetData>& Ass
 	if ((Blueprint != nullptr) && FBlueprintEditorUtils::IsActorBased(Blueprint))
 	{
 		OutTooltipText = LOCTEXT("UnsupportedAssetTypeForGraphDragDrop", "Cannot create a node from this type of asset").ToString();
+		const FBlueprintGraphModule& Module = FModuleManager::LoadModuleChecked<FBlueprintGraphModule>("BlueprintGraph");
+		bool bFoundCustomText = false;
 		for (const FAssetData& AssetData : Assets)
 		{
+			// check asset to see if we can get some AssetBlueprintGraphActions for it
+			if (const FAssetBlueprintGraphActions* GraphActions = Module.GetAssetBlueprintGraphActions(AssetData.GetClass()))
+			{
+				// get the text from the module we loaded
+				FText CustomText = GraphActions->GetGraphHoverMessage(AssetData, HoverGraph);
+				if (!CustomText.IsEmpty())
+				{
+					// want to make sure the hover message properly represents that Input Action nodes can only be dragged onto event graphs
+					if (FBlueprintEditorUtils::IsEventGraph(HoverGraph))
+					{
+						OutOkIcon = true;
+						OutTooltipText = CustomText.ToString();
+					}
+					else
+					{
+						OutOkIcon = false;
+						OutTooltipText = LOCTEXT("UnsupportedAssetTypeForGraphDragDropEventGraph", "Cannot create a node from this type of asset in this graph").ToString();
+					}
+					return;
+				}
+			}
+
 			if (UObject* Asset = AssetData.GetAsset())
 			{
 				UClass* AssetClass = Asset->GetClass();
diff --git a/Engine/Source/Editor/BlueprintGraph/Public/BlueprintGraphModule.h b/Engine/Source/Editor/BlueprintGraph/Public/BlueprintGraphModule.h
--- a/Engine/Source/Editor/BlueprintGraph/Public/BlueprintGraphModule.h
+++ b/Engine/Source/Editor/BlueprintGraph/Public/BlueprintGraphModule.h
@@ -9,6 +9,7 @@
 
 class FBlueprintActionFilter;
 struct FBlueprintActionInfo;
+struct FAssetBlueprintGraphActions;
 
 class FBlueprintGraphModule : public IModuleInterface
 {
@@ -28,6 +29,13 @@ public:
 	BLUEPRINTGRAPH_API TArray<FActionMenuRejectionTest>& GetExtendedActionMenuFilters() { return ExtendedMenuFilters; }
 
 	virtual void ShutdownModule() override;
+	BLUEPRINTGRAPH_API void RegisterGraphAction(const UClass* AssetType, TUniquePtr<FAssetBlueprintGraphActions> Action);
+	BLUEPRINTGRAPH_API void UnregisterGraphAction(const UClass* AssetType);
+	const FAssetBlueprintGraphActions* GetAssetBlueprintGraphActions(const UClass* AssetType) const;
+
 private:
+	/** Customizations for specific assets that want to customize their interactions with the Blueprint Graph window */
+	TMap<const UClass*, TUniquePtr< FAssetBlueprintGraphActions>> AssetBlueprintGraphActions;
+
 	TArray<FActionMenuRejectionTest> ExtendedMenuFilters;
 };
