Updating Input Mapping Contexts when Input Action triggers/modifiers are modifiied

diff --git a/Engine/Plugins/EnhancedInput/Source/EnhancedInput/Private/InputAction.cpp b/Engine/Plugins/EnhancedInput/Source/EnhancedInput/Private/InputAction.cpp
--- a/Engine/Plugins/EnhancedInput/Source/EnhancedInput/Private/InputAction.cpp
+++ b/Engine/Plugins/EnhancedInput/Source/EnhancedInput/Private/InputAction.cpp
@@ -110,6 +110,16 @@ TSet<const UInputAction*> UInputAction::ActionsWithModifiedValueTypes;
 
 void UInputAction::PostEditChangeProperty(FPropertyChangedEvent& PropertyChangedEvent)
 {
+	// if our value changes were to the trigger or modifier array broadcast the change
+	if (PropertyChangedEvent.MemberProperty->GetFName() == GET_MEMBER_NAME_CHECKED(UInputAction, Modifiers))
+	{
+		OnModifiersChanged.Broadcast();
+	}
+	if (PropertyChangedEvent.MemberProperty->GetFName() == GET_MEMBER_NAME_CHECKED(UInputAction, Triggers))
+	{
+		OnTriggersChanged.Broadcast();
+	}
+	
 	// If our value type changes we need to inform any blueprint InputActionEx nodes that refer to this action
 	if (PropertyChangedEvent.MemberProperty->GetFName() == GET_MEMBER_NAME_CHECKED(UInputAction, ValueType) ||
 		PropertyChangedEvent.MemberProperty->GetFName() == GET_MEMBER_NAME_CHECKED(UInputAction, Triggers))
diff --git a/Engine/Plugins/EnhancedInput/Source/EnhancedInput/Public/InputAction.h b/Engine/Plugins/EnhancedInput/Source/EnhancedInput/Public/InputAction.h
--- a/Engine/Plugins/EnhancedInput/Source/EnhancedInput/Public/InputAction.h
+++ b/Engine/Plugins/EnhancedInput/Source/EnhancedInput/Public/InputAction.h
@@ -34,6 +34,12 @@ public:
 	 * Returns a bitmask of supported trigger events that is built from each UInputTrigger on this Action.
 	 */
 	ETriggerEventsSupported GetSupportedTriggerEvents() const;
+
+	DECLARE_MULTICAST_DELEGATE(FTriggersChanged);
+	DECLARE_MULTICAST_DELEGATE(FModifiersChanged);
+
+	FTriggersChanged OnTriggersChanged;
+	FModifiersChanged OnModifiersChanged;
 #endif // WITH_EDITOR
 	
 	/**
diff --git a/Engine/Plugins/EnhancedInput/Source/InputEditor/Private/InputCustomizations.cpp b/Engine/Plugins/EnhancedInput/Source/InputEditor/Private/InputCustomizations.cpp
--- a/Engine/Plugins/EnhancedInput/Source/InputEditor/Private/InputCustomizations.cpp
+++ b/Engine/Plugins/EnhancedInput/Source/InputEditor/Private/InputCustomizations.cpp
@@ -39,6 +39,20 @@ void FInputContextDetails::CustomizeDetails(class IDetailLayoutBuilder& DetailBu
 //////////////////////////////////////////////////////////
 // FEnhancedActionMappingCustomization
 
+const FEnhancedActionKeyMapping* GetActionKeyMapping(TSharedPtr<IPropertyHandle> MappingPropertyHandle)
+{
+	if (MappingPropertyHandle->IsValidHandle())
+	{
+		void* Data = nullptr;
+		if (MappingPropertyHandle->GetValueData(Data) != FPropertyAccess::Fail)
+		{
+			return static_cast<FEnhancedActionKeyMapping*>(Data);
+		}
+	}
+
+	return nullptr;
+}
+
 void FEnhancedActionMappingCustomization::CustomizeHeader(TSharedRef<IPropertyHandle> PropertyHandle, FDetailWidgetRow& HeaderRow, IPropertyTypeCustomizationUtils& CustomizationUtils)
 {
 	MappingPropertyHandle = PropertyHandle;
@@ -65,38 +79,52 @@ void FEnhancedActionMappingCustomization::CustomizeHeader(TSharedRef<IPropertyHa
 	KeyStructCustomization->CustomizeHeaderOnlyWithButton(KeyHandle.ToSharedRef(), HeaderRow, CustomizationUtils, RemoveButton);
 }
 
-void AddInputActionProperties(TSharedRef<IPropertyHandle> PropertyHandle, IDetailChildrenBuilder& ChildBuilder)
+void FEnhancedActionMappingCustomization::AddInputActionProperties(TSharedRef<IPropertyHandle> PropertyHandle, IDetailChildrenBuilder& ChildBuilder)
 {
-	void* Data = nullptr;
-    if (PropertyHandle->GetValueData(Data) != FPropertyAccess::Fail)
+    if (const FEnhancedActionKeyMapping* ActionKeyMapping = GetActionKeyMapping(PropertyHandle))
     {
-    	FEnhancedActionKeyMapping* ActionKeyMapping = static_cast<FEnhancedActionKeyMapping*>(Data);
-    	
     	// Make sure ActionKeyMapping action is valid. If triggers and modifiers are empty we can just back out - nothing to add
-    	if (ActionKeyMapping->Action && (!ActionKeyMapping->Action->Triggers.IsEmpty() || !ActionKeyMapping->Action->Modifiers.IsEmpty()))
+    	if (ActionKeyMapping->Action)
     	{
     		// Convert InputAction to non const so we can properly use it in the UObject Array for AddExternalObjectProperty
     		const UInputAction* InputActionPtr = ActionKeyMapping->Action;
     		UInputAction* InputAction = const_cast<UInputAction*>(InputActionPtr);
     		TArray<UObject*> ActionsAsUObjects { InputAction };
-    		// no need to do it if the trigger array is empty
-    		if (!ActionKeyMapping->Action->Triggers.IsEmpty())
+
+    		if (IDetailPropertyRow* InputActionTriggersRow = ChildBuilder.AddExternalObjectProperty(ActionsAsUObjects, GET_MEMBER_NAME_CHECKED(UInputAction, Triggers)))
     		{
-    			if (IDetailPropertyRow* InputActionTriggersRow = ChildBuilder.AddExternalObjectProperty(ActionsAsUObjects, GET_MEMBER_NAME_CHECKED(UInputAction, Triggers)))
+    			InputActionTriggersPropertyRow = InputActionTriggersRow;
+    			InputActionTriggersRow->DisplayName(LOCTEXT("InputActionTriggersDisplayName", "Triggers From Input Action"));
+    			InputActionTriggersRow->ToolTip(FText::Format(LOCTEXT("InputActionTriggersToolTip", "Triggers from the {0} Input Action"), FText::FromName(ActionKeyMapping->Action.GetFName())));
+    			InputActionTriggersRow->IsEnabled(false);
+    			InputAction->OnTriggersChanged.AddSP(this, &FEnhancedActionMappingCustomization::OnInputActionTriggersChanged);
+    			
+    			// hide it if the trigger array is empty
+    			if (ActionKeyMapping->Action->Triggers.IsEmpty())
     			{
-    				InputActionTriggersRow->DisplayName(LOCTEXT("InputActionTriggersDisplayName", "Triggers From Input Action"));
-    				InputActionTriggersRow->ToolTip(FText::Format(LOCTEXT("InputActionTriggersToolTip", "Triggers from the {0} Input Action"), FText::FromName(ActionKeyMapping->Action.GetFName())));
-    				InputActionTriggersRow->IsEnabled(false);
+					InputActionTriggersRow->Visibility(EVisibility::Hidden);
+    			}
+    			else
+    			{
+    				InputActionTriggersRow->Visibility(EVisibility::Visible);
     			}
     		}
-    		// no need to do it if the modifier array is empty
-    		if (!ActionKeyMapping->Action->Modifiers.IsEmpty())
+    		if (IDetailPropertyRow* InputActionModifiersRow = ChildBuilder.AddExternalObjectProperty(ActionsAsUObjects, GET_MEMBER_NAME_CHECKED(UInputAction, Modifiers)))
     		{
-    			if (IDetailPropertyRow* InputActionModifiersRow = ChildBuilder.AddExternalObjectProperty(ActionsAsUObjects, GET_MEMBER_NAME_CHECKED(UInputAction, Modifiers)))
+    			InputActionModifiersPropertyRow = InputActionModifiersRow;
+    			InputActionModifiersRow->DisplayName(LOCTEXT("InputActionModifiersDisplayName", "Modifiers From Input Action"));
+    			InputActionModifiersRow->ToolTip(FText::Format(LOCTEXT("InputActionModifiersToolTip", "Modifiers from the {0} Input Action"), FText::FromName(ActionKeyMapping->Action.GetFName())));
+    			InputActionModifiersRow->IsEnabled(false);
+    			InputAction->OnModifiersChanged.AddSP(this, &FEnhancedActionMappingCustomization::OnInputActionModifiersChanged);
+    			
+    			// hide it if the modifier array is empty
+    			if (ActionKeyMapping->Action->Modifiers.IsEmpty())
+    			{
+    				InputActionModifiersRow->Visibility(EVisibility::Hidden);
+    			}
+    			else
     			{
-    				InputActionModifiersRow->DisplayName(LOCTEXT("InputActionModifiersDisplayName", "Modifiers From Input Action"));
-    				InputActionModifiersRow->ToolTip(FText::Format(LOCTEXT("InputActionModifiersToolTip", "Modifiers from the {0} Input Action"), FText::FromName(ActionKeyMapping->Action.GetFName())));
-    				InputActionModifiersRow->IsEnabled(false);
+    				InputActionModifiersRow->Visibility(EVisibility::Visible);	
     			}
     		}
     	}
@@ -140,6 +168,46 @@ void FEnhancedActionMappingCustomization::OnTriggersChanged() const
     KeyStructCustomization->SetEnableKeySelector(!bContainsComboTrigger);
 }
 
+void FEnhancedActionMappingCustomization::OnInputActionTriggersChanged() const
+{
+	if (const FEnhancedActionKeyMapping* ActionKeyMapping = GetActionKeyMapping(MappingPropertyHandle))
+	{
+		// Make sure ActionKeyMapping action and the property row are valid
+		if (ActionKeyMapping->Action && InputActionTriggersPropertyRow)
+		{
+			// if so we want to hide the row or show it based off contents of the array (whether it's empty or not)
+			if (!ActionKeyMapping->Action->Triggers.IsEmpty())
+			{
+				InputActionTriggersPropertyRow->Visibility(EVisibility::Visible);
+			}
+			else
+			{
+				InputActionTriggersPropertyRow->Visibility(EVisibility::Hidden);
+			}
+		}
+	}
+}
+
+void FEnhancedActionMappingCustomization::OnInputActionModifiersChanged() const
+{
+    if (const FEnhancedActionKeyMapping* ActionKeyMapping = GetActionKeyMapping(MappingPropertyHandle))
+    {
+    	// Make sure ActionKeyMapping action and the property row are valid
+    	if (ActionKeyMapping->Action && InputActionModifiersPropertyRow)
+    	{
+    		// if so we want to hide the row or show it based off contents of the array (whether it's empty or not)
+    		if (!ActionKeyMapping->Action->Modifiers.IsEmpty())
+    		{
+    			InputActionModifiersPropertyRow->Visibility(EVisibility::Visible);
+    		}
+    		else
+    		{
+    			InputActionModifiersPropertyRow->Visibility(EVisibility::Hidden);
+    		}
+    	}
+    }
+}
+
 bool FEnhancedActionMappingCustomization::DoesTriggerArrayContainCombo() const
 {
 	if (MappingPropertyHandle)
diff --git a/Engine/Plugins/EnhancedInput/Source/InputEditor/Private/InputCustomizations.h b/Engine/Plugins/EnhancedInput/Source/InputEditor/Private/InputCustomizations.h
--- a/Engine/Plugins/EnhancedInput/Source/InputEditor/Private/InputCustomizations.h
+++ b/Engine/Plugins/EnhancedInput/Source/InputEditor/Private/InputCustomizations.h
@@ -8,6 +8,7 @@
 #include "IPropertyTypeCustomization.h"
 #include "KeyStructCustomization.h"
 
+class IDetailPropertyRow;
 class IDetailLayoutBuilder;
 
 class FInputContextDetails : public IDetailCustomization
@@ -38,10 +39,15 @@ private:
 
 	void RemoveMappingButton_OnClick() const;
 	void OnTriggersChanged() const;
+	void AddInputActionProperties(TSharedRef<IPropertyHandle> PropertyHandle, IDetailChildrenBuilder& ChildBuilder);
+	void OnInputActionTriggersChanged() const;
+	void OnInputActionModifiersChanged() const;
 
 	TSharedPtr<IPropertyTypeCustomization> KeyStructInstance;
 	TSharedPtr<IPropertyHandle> MappingPropertyHandle;
 	TSharedPtr<FKeyStructCustomization> KeyStructCustomization;
+	IDetailPropertyRow* InputActionTriggersPropertyRow = nullptr;
+	IDetailPropertyRow* InputActionModifiersPropertyRow = nullptr;
 };
 
 /**
