Creating an Enhanced Input World subsystem

diff --git a/Engine/Plugins/EnhancedInput/Source/EnhancedInput/Private/EnhancedInputDeveloperSettings.cpp b/Engine/Plugins/EnhancedInput/Source/EnhancedInput/Private/EnhancedInputDeveloperSettings.cpp
--- a/Engine/Plugins/EnhancedInput/Source/EnhancedInput/Private/EnhancedInputDeveloperSettings.cpp
+++ b/Engine/Plugins/EnhancedInput/Source/EnhancedInput/Private/EnhancedInputDeveloperSettings.cpp
@@ -2,9 +2,23 @@
 
 #include "EnhancedInputDeveloperSettings.h"
 #include "EnhancedInputPlatformSettings.h"
+#include "HAL/IConsoleManager.h"
 
 #include UE_INLINE_GENERATED_CPP_BY_NAME(EnhancedInputDeveloperSettings)
 
+namespace UE::EnhancedInput::Private
+{
+	namespace ConsoleVariables
+	{
+		static bool bShouldLogAllWorldSubsystemInputs = false;
+		static FAutoConsoleVariableRef CVarShouldLogAllWorldSubsystemInputs(
+			TEXT("EnhancedInput.bShouldLogAllWorldSubsystemInputs"),
+			bShouldLogAllWorldSubsystemInputs,
+			TEXT("Should each InputKey call to the World subsystem be logged?"),
+			ECVF_Cheat);
+	}
+}
+
 UEnhancedInputDeveloperSettings::UEnhancedInputDeveloperSettings(const FObjectInitializer& Initializer)
 	: Super(Initializer)
 {
diff --git a/Engine/Plugins/EnhancedInput/Source/EnhancedInput/Private/EnhancedInputSubsystems.cpp b/Engine/Plugins/EnhancedInput/Source/EnhancedInput/Private/EnhancedInputSubsystems.cpp
--- a/Engine/Plugins/EnhancedInput/Source/EnhancedInput/Private/EnhancedInputSubsystems.cpp
+++ b/Engine/Plugins/EnhancedInput/Source/EnhancedInput/Private/EnhancedInputSubsystems.cpp
@@ -4,17 +4,22 @@
 
 #include "Engine/LocalPlayer.h"
 #include "EnhancedInputComponent.h"
+#include "EnhancedInputDeveloperSettings.h"
 #include "EnhancedInputModule.h"
 #include "EnhancedPlayerInput.h"
+#include "Engine/World.h"
+#include "Framework/Application/SlateApplication.h"
 #include "GameFramework/PlayerController.h"
 #include "InputCoreTypes.h"
 #include "InputMappingContext.h"
 #include "InputModifiers.h"
 #include "InputTriggers.h"
+#include "GameFramework/InputSettings.h"
 #include "UObject/UObjectIterator.h"
 
 #include UE_INLINE_GENERATED_CPP_BY_NAME(EnhancedInputSubsystems)
 
+DEFINE_LOG_CATEGORY(LogWorldSubsystemInput);
 
 // **************************************************************************************************
 // *
@@ -35,3 +40,227 @@ void UEnhancedInputLocalPlayerSubsystem::ControlMappingsRebuiltThisFrame()
 {
 	ControlMappingsRebuiltDelegate.Broadcast();
 }
+
+// **************************************************************************************************
+// *
+// * UEnhancedInputWorldSubsystem
+// *
+// **************************************************************************************************
+
+void UEnhancedInputWorldSubsystem::Initialize(FSubsystemCollectionBase& Collection)
+{
+	Super::Initialize(Collection);
+
+	UClass* PlayerInputClass = nullptr;
+	if (TSoftClassPtr<UEnhancedPlayerInput> DefaultWorldClass = GetDefault<UEnhancedInputDeveloperSettings>()->DefaultWorldInputClass)
+	{
+		PlayerInputClass = DefaultWorldClass.Get();
+	}
+	else
+	{
+		PlayerInputClass = UEnhancedPlayerInput::StaticClass();
+	}
+		
+	PlayerInput = NewObject<UEnhancedPlayerInput>(this, PlayerInputClass, TEXT("EIWorldSubsystem_PlayerInput0"));	
+	ensureMsgf(PlayerInput, TEXT("UEnhancedInputWorldSubsystem::Initialize failed to create PlayerInput! This subsystem will not tick!"));
+	
+	if (FSlateApplication::IsInitialized())
+	{
+		InputPreprocessor = MakeShared<FEnhancedInputWorldProcessor>();
+		InputPreprocessor->SetSubsystemWorld(GetWorld());
+		FSlateApplication::Get().RegisterInputPreProcessor(InputPreprocessor, 0);
+	}
+	
+	StartConsumingInput();
+}
+
+void UEnhancedInputWorldSubsystem::Deinitialize()
+{
+	Super::Deinitialize();
+
+	if (FSlateApplication::IsInitialized())
+	{
+		FSlateApplication::Get().UnregisterInputPreProcessor(InputPreprocessor);
+	}
+	
+	InputPreprocessor.Reset();
+	CurrentInputStack.Empty();
+	PlayerInput = nullptr;
+}
+
+bool UEnhancedInputWorldSubsystem::ShouldCreateSubsystem(UObject* Outer) const
+{
+	if (!FSlateApplication::IsInitialized())
+	{
+		return false;
+	}
+	
+	return Super::ShouldCreateSubsystem(Outer);
+}
+
+ETickableTickType UEnhancedInputWorldSubsystem::GetTickableTickType() const
+{
+	return IsTemplate() ? ETickableTickType::Never : ETickableTickType::Always;
+}
+
+void UEnhancedInputWorldSubsystem::Tick(float DeltaTime)
+{
+	if (!ensure(PlayerInput))
+	{
+		return;
+	}
+
+	FModifyContextOptions Options;
+	Options.bForceImmediately = true;
+	
+	// Rebuild the control mappings and tick any forced input that may have been injected
+	RequestRebuildControlMappings(Options);
+	TickForcedInput(DeltaTime);
+
+	static TArray<UInputComponent*> InputStack;
+	// Build the input stack
+	{
+		for (int32 i = 0; i < CurrentInputStack.Num(); ++i)
+		{
+			if (UInputComponent* IC = CurrentInputStack[i].Get())
+			{
+				InputStack.Push(IC);
+			}
+			else
+			{
+				CurrentInputStack.RemoveAt(i--);
+			}
+		}
+	}
+	
+	// Process input stack on the player input
+	PlayerInput->Tick(DeltaTime);
+	PlayerInput->ProcessInputStack(InputStack, DeltaTime, GetWorld()->IsPaused());
+
+	InputStack.Reset();
+}
+
+UEnhancedPlayerInput* UEnhancedInputWorldSubsystem::GetPlayerInput() const
+{
+	return PlayerInput;
+}
+
+void UEnhancedInputWorldSubsystem::PushInputComponent(AActor* Actor)
+{
+	if (Actor)
+	{
+		if (TObjectPtr<UInputComponent> InInputComponent = Actor->InputComponent)
+		{
+			bool bPushed = false;
+			CurrentInputStack.RemoveSingle(InInputComponent);
+			for (int32 Index = CurrentInputStack.Num() - 1; Index >= 0; --Index)
+			{
+				UInputComponent* IC = CurrentInputStack[Index].Get();
+				if (IC == nullptr)
+				{
+					CurrentInputStack.RemoveAt(Index);
+				}
+				else if (IC->Priority <= InInputComponent->Priority)
+				{
+					CurrentInputStack.Insert(InInputComponent, Index + 1);
+					bPushed = true;
+					break;
+				}
+			}
+			if (!bPushed)
+			{
+				CurrentInputStack.Insert(InInputComponent, 0);
+				RequestRebuildControlMappings();
+			}
+		}
+		else
+		{
+			UE_LOG(LogWorldSubsystemInput, Error, TEXT("Attempted to push a null Input Component to the Enhanced Input World Subsystem!"));
+		}
+	}
+}
+
+bool UEnhancedInputWorldSubsystem::PopInputComponent(AActor* Actor)
+{
+	if (Actor) 
+	{
+		if (TObjectPtr<UInputComponent> InInputComponent = Actor->InputComponent)
+    	{
+    		if (CurrentInputStack.RemoveSingle(InInputComponent) > 0)
+    		{
+    			InInputComponent->ClearBindingValues();
+    			RequestRebuildControlMappings();
+    			return true;
+    		}
+    	}
+	}
+    return false;
+}
+
+void UEnhancedInputWorldSubsystem::StartConsumingInput()
+{
+	bIsCurrentlyConsumingInput = true;
+	
+	AddDefaultMappingContexts();
+
+	RequestRebuildControlMappings();
+}
+
+void UEnhancedInputWorldSubsystem::StopConsumingInput()
+{
+	bIsCurrentlyConsumingInput = false;
+
+	RemoveDefaultMappingContexts();
+
+	RequestRebuildControlMappings();
+}
+
+bool UEnhancedInputWorldSubsystem::InputKey(const FInputKeyParams& Params)
+{
+	if (PlayerInput)
+	{
+		if (bIsCurrentlyConsumingInput)
+		{
+			if (GetDefault<UEnhancedInputDeveloperSettings>()->bShouldLogAllWorldSubsystemInputs)
+			{
+				UE_LOG(LogWorldSubsystemInput, Log, TEXT("EI World Subsystem InputKey : [%s]"), *Params.Key.ToString());
+			}
+			return PlayerInput->InputKey(Params);
+		}
+		return false;
+	}
+
+	ensureAlwaysMsgf(false, TEXT("Attempting to input a key to the EnhancedInputWorldSubsystem, but there is no Player Input!"));
+	return false;
+}
+
+void UEnhancedInputWorldSubsystem::AddDefaultMappingContexts()
+{
+	if (GetDefault<UEnhancedInputDeveloperSettings>()->bEnableDefaultMappingContexts)
+	{
+		for (const FDefaultContextSetting& ContextSetting : GetMutableDefault<UEnhancedInputDeveloperSettings>()->DefaultWorldSubsystemMappingContexts)
+		{
+			if (const UInputMappingContext* IMC = ContextSetting.InputMappingContext.LoadSynchronous())
+			{
+				if (!HasMappingContext(IMC))
+				{
+					AddMappingContext(IMC, ContextSetting.Priority);
+				}
+			}
+		}
+	}
+}
+
+void UEnhancedInputWorldSubsystem::RemoveDefaultMappingContexts()
+{
+	if (GetDefault<UEnhancedInputDeveloperSettings>()->bEnableDefaultMappingContexts)
+	{
+		for (const FDefaultContextSetting& ContextSetting : GetMutableDefault<UEnhancedInputDeveloperSettings>()->DefaultWorldSubsystemMappingContexts)
+		{
+			if (const UInputMappingContext* IMC = ContextSetting.InputMappingContext.LoadSynchronous())
+			{
+				RemoveMappingContext(IMC);
+			}
+		}
+	}
+}
diff --git a/Engine/Plugins/EnhancedInput/Source/EnhancedInput/Private/EnhancedInputWorldProcessor.cpp b/Engine/Plugins/EnhancedInput/Source/EnhancedInput/Private/EnhancedInputWorldProcessor.cpp
new file mode 100644
--- /dev/null
+++ b/Engine/Plugins/EnhancedInput/Source/EnhancedInput/Private/EnhancedInputWorldProcessor.cpp
@@ -0,0 +1,180 @@
+// Copyright Epic Games, Inc. All Rights Reserved.
+
+#include "EnhancedInputWorldProcessor.h"
+#include "EnhancedInputSubsystems.h"
+#include "Engine/World.h"
+#include "Framework/Application/SlateApplication.h"
+#include "GameFramework/PlayerInput.h"
+
+void FEnhancedInputWorldProcessor::Tick(const float DeltaTime, FSlateApplication& SlateApp, TSharedRef<ICursor> Cursor)
+{
+	ProcessAccumulatedPointerInput(DeltaTime);
+}
+
+bool FEnhancedInputWorldProcessor::HandleKeyDownEvent(FSlateApplication& SlateApp, const FKeyEvent& InKeyEvent)
+{
+	FInputKeyParams Params = {};
+	Params.Key = InKeyEvent.GetKey();
+	Params.Event = IE_Pressed;
+	Params.Delta.X = 1.0;
+	Params.DeltaTime = SlateApp.GetDeltaTime();
+	Params.NumSamples = Params.Key.IsAnalog() ? 1 : 0;
+	
+	InputKeyToSubsystem(Params);
+	
+	return IInputProcessor::HandleKeyDownEvent(SlateApp, InKeyEvent);
+}
+
+bool FEnhancedInputWorldProcessor::HandleKeyUpEvent(FSlateApplication& SlateApp, const FKeyEvent& InKeyEvent)
+{
+	FInputKeyParams Params = {};
+	Params.Key = InKeyEvent.GetKey();
+	Params.Event = IE_Released;
+	Params.Delta.X = 0.0;
+	Params.DeltaTime = SlateApp.GetDeltaTime();
+	Params.NumSamples = Params.Key.IsAnalog() ? 1 : 0;
+	
+	InputKeyToSubsystem(Params);
+	return IInputProcessor::HandleKeyUpEvent(SlateApp, InKeyEvent);
+}
+
+bool FEnhancedInputWorldProcessor::HandleAnalogInputEvent(FSlateApplication& SlateApp, const FAnalogInputEvent& InAnalogInputEvent)
+{
+	FInputKeyParams Params = {};
+	Params.Key = InAnalogInputEvent.GetKey();
+	Params.Event = IE_Pressed;
+	Params.Delta.X = InAnalogInputEvent.GetAnalogValue();
+	Params.DeltaTime = SlateApp.GetDeltaTime();
+	Params.NumSamples = 1;
+	
+	InputKeyToSubsystem(Params);
+	return IInputProcessor::HandleAnalogInputEvent(SlateApp, InAnalogInputEvent);
+}
+
+bool FEnhancedInputWorldProcessor::HandleMouseMoveEvent(FSlateApplication& SlateApp, const FPointerEvent& MouseEvent)
+{
+	UpdateCachedPointerPosition(MouseEvent);
+	
+	return IInputProcessor::HandleMouseMoveEvent(SlateApp, MouseEvent);
+}
+
+bool FEnhancedInputWorldProcessor::HandleMouseButtonDownEvent(FSlateApplication& SlateApp, const FPointerEvent& MouseEvent)
+{
+	FInputKeyParams Params = {};
+	Params.Key = MouseEvent.GetEffectingButton();
+	Params.Event = IE_Pressed;
+	Params.Delta.X = 1.0;
+	Params.DeltaTime = SlateApp.GetDeltaTime();
+	Params.NumSamples = 0;
+	
+	InputKeyToSubsystem(Params);
+	return IInputProcessor::HandleMouseButtonDownEvent(SlateApp, MouseEvent);
+}
+
+bool FEnhancedInputWorldProcessor::HandleMouseButtonUpEvent(FSlateApplication& SlateApp, const FPointerEvent& MouseEvent)
+{
+	FInputKeyParams Params = {};
+	Params.Key = MouseEvent.GetEffectingButton();
+	Params.Event = IE_Released;
+	Params.Delta.X = 0.0;
+	Params.DeltaTime = SlateApp.GetDeltaTime();
+	Params.NumSamples = 0;
+	
+	InputKeyToSubsystem(Params);
+	return IInputProcessor::HandleMouseButtonUpEvent(SlateApp, MouseEvent);
+}
+
+bool FEnhancedInputWorldProcessor::HandleMouseButtonDoubleClickEvent(FSlateApplication& SlateApp, const FPointerEvent& MouseEvent)
+{
+	FInputKeyParams Params = {};
+	Params.Key = MouseEvent.GetEffectingButton();
+	Params.Event = IE_DoubleClick;
+	Params.Delta.X = 1.0;
+	Params.DeltaTime = SlateApp.GetDeltaTime();
+	Params.NumSamples = 0;
+	
+	InputKeyToSubsystem(Params);
+	return IInputProcessor::HandleMouseButtonDoubleClickEvent(SlateApp, MouseEvent);
+}
+
+bool FEnhancedInputWorldProcessor::HandleMouseWheelOrGestureEvent(FSlateApplication& SlateApp, const FPointerEvent& InWheelEvent, const FPointerEvent* InGestureEvent)
+{
+	const FKey MouseWheelKey = InWheelEvent.GetWheelDelta() < 0 ? EKeys::MouseScrollDown : EKeys::MouseScrollUp;
+
+	// Input the Mouse wheel key events (mouse scroll down or scroll up) as being pressed and released this frame
+	// The SceneViewport Inputs the Mouse Scroll Wheel buttons up and down in the same frame, this replicates that behavior
+	{
+		FInputKeyParams PressedParams = {};
+		PressedParams.Key = MouseWheelKey;
+		PressedParams.Event = IE_Pressed;
+		PressedParams.Delta.X = 1.0;
+		PressedParams.DeltaTime = SlateApp.GetDeltaTime();
+		PressedParams.NumSamples = 0;
+		
+		FInputKeyParams ReleasedParams = PressedParams;
+		ReleasedParams.Event = IE_Released;
+
+		InputKeyToSubsystem(PressedParams);
+		InputKeyToSubsystem(ReleasedParams);	
+	}
+	// Input the wheel axis delta to get the MouseWheelAxis button working
+	{
+		FInputKeyParams Params = {};
+		Params.Key = EKeys::MouseWheelAxis;
+		Params.Delta.X = InWheelEvent.GetWheelDelta();
+		Params.DeltaTime = SlateApp.GetDeltaTime();
+		Params.NumSamples = 1;
+		
+		InputKeyToSubsystem(Params);
+	}
+	
+	return IInputProcessor::HandleMouseWheelOrGestureEvent(SlateApp, InWheelEvent, InGestureEvent);
+}
+
+void FEnhancedInputWorldProcessor::UpdateCachedPointerPosition(const FPointerEvent& MouseEvent)
+{
+	CachedCursorDelta = MouseEvent.GetCursorDelta();
+	
+	++NumCursorSamplesThisFrame.X;
+	++NumCursorSamplesThisFrame.Y;
+}
+
+void FEnhancedInputWorldProcessor::ProcessAccumulatedPointerInput(float DeltaTime)
+{
+	// Input the MouseX value
+	{
+		FInputKeyParams Params = {};
+		Params.Key = EKeys::MouseX;
+		Params.Delta.X = CachedCursorDelta.X;
+		Params.DeltaTime = DeltaTime;
+		Params.NumSamples = NumCursorSamplesThisFrame.X;
+		InputKeyToSubsystem(Params);
+	}
+
+	// Input the MouseY value
+	{
+		FInputKeyParams Params = {};
+		Params.Key = EKeys::MouseY;
+		Params.Delta.X = CachedCursorDelta.Y;
+		Params.DeltaTime = DeltaTime;
+		Params.NumSamples = NumCursorSamplesThisFrame.Y;
+		InputKeyToSubsystem(Params);
+	}
+	
+	NumCursorSamplesThisFrame = FIntPoint::ZeroValue;
+	CachedCursorDelta = FVector2D::ZeroVector;
+}
+
+
+bool FEnhancedInputWorldProcessor::InputKeyToSubsystem(const FInputKeyParams& Params)
+{
+	if (WorldFromSubsystem)
+	{
+		if (UEnhancedInputWorldSubsystem* WorldSubsystem = UWorld::GetSubsystem<UEnhancedInputWorldSubsystem>(WorldFromSubsystem))
+		{
+			return WorldSubsystem->InputKey(Params);
+		}
+	}
+	
+	return false;
+}
diff --git a/Engine/Plugins/EnhancedInput/Source/EnhancedInput/Public/EnhancedInputDeveloperSettings.h b/Engine/Plugins/EnhancedInput/Source/EnhancedInput/Public/EnhancedInputDeveloperSettings.h
--- a/Engine/Plugins/EnhancedInput/Source/EnhancedInput/Public/EnhancedInputDeveloperSettings.h
+++ b/Engine/Plugins/EnhancedInput/Source/EnhancedInput/Public/EnhancedInputDeveloperSettings.h
@@ -43,6 +43,17 @@ public:
 	 */
 	UPROPERTY(config, EditAnywhere, Category = "Enhanced Input",  meta = (editCondition = "bEnableDefaultMappingContexts"))
 	TArray<FDefaultContextSetting> DefaultMappingContexts;
+	
+	/**
+	 * Array of any input mapping contexts that you want to be applied by default to the Enhanced Input world subsystem.
+	 * NOTE: These mapping context's can only be from your game's root content directory, not plugins.
+	 */
+	UPROPERTY(config, EditAnywhere, Category = "Enhanced Input", meta = (editCondition = "bEnableDefaultMappingContexts"))
+	TArray<FDefaultContextSetting> DefaultWorldSubsystemMappingContexts;
+	
+	/** The default player input class that the Enhanced Input world subsystem will use. */
+	UPROPERTY(config, EditAnywhere, NoClear, Category = "Enhanced Input")
+	TSoftClassPtr<class UEnhancedPlayerInput> DefaultWorldInputClass;
 
 	/**
 	 * Platform specific settings for Enhanced Input.
@@ -59,4 +70,11 @@ public:
 	 */
 	UPROPERTY(EditAnywhere, Category = "Enhanced Input", meta=(ConsoleVariable="EnhancedInput.OnlyTriggerLastActionInChord"))
 	bool bShouldOnlyTriggerLastActionInChord = true;
+	
+	/**
+ 	 * If true then the Enhanced Input world subsystem will log all input that is being processed by it (keypresses, analog values, etc)
+ 	 * Note: This can produce A LOT of logs, so only use this if you are debugging something.
+ 	 */
+	UPROPERTY(config, EditAnywhere, Category = "Enhanced Input World Subsystem", meta=(ConsoleVariable="EnhancedInput.bShouldLogAllWorldSubsystemInputs"))
+	uint8 bShouldLogAllWorldSubsystemInputs : 1;
 };
\ No newline at end of file
diff --git a/Engine/Plugins/EnhancedInput/Source/EnhancedInput/Public/EnhancedInputSubsystems.h b/Engine/Plugins/EnhancedInput/Source/EnhancedInput/Public/EnhancedInputSubsystems.h
--- a/Engine/Plugins/EnhancedInput/Source/EnhancedInput/Public/EnhancedInputSubsystems.h
+++ b/Engine/Plugins/EnhancedInput/Source/EnhancedInput/Public/EnhancedInputSubsystems.h
@@ -4,11 +4,15 @@
 
 #include "CoreMinimal.h"
 #include "EnhancedInputSubsystemInterface.h"
+#include "EnhancedInputWorldProcessor.h"
 #include "Subsystems/LocalPlayerSubsystem.h"
 #include "Subsystems/EngineSubsystem.h"
+#include "Subsystems/WorldSubsystem.h"
 
 #include "EnhancedInputSubsystems.generated.h"
 
+DECLARE_LOG_CATEGORY_EXTERN(LogWorldSubsystemInput, Log, All);
+
 // Per local player input subsystem
 UCLASS()
 class ENHANCEDINPUT_API UEnhancedInputLocalPlayerSubsystem : public ULocalPlayerSubsystem, public IEnhancedInputSubsystemInterface
@@ -31,3 +35,84 @@ public:
 	UPROPERTY(BlueprintAssignable, DisplayName=OnControlMappingsRebuilt, Category = "Input")
 	FOnControlMappingsRebuilt ControlMappingsRebuiltDelegate;
 };
+
+/**
+ * Per world input subsystem that allows you to bind input delegates to actors without an owning Player Controller. 
+ * This should be used when an actor needs to receive input delegates but will never have an owning Player Controller.
+ * For example, you can add input delegates to unlock a door when the user has a certain set of keys pressed.
+ * Be sure to enable input on the actor, or else the input delegates won't fire!
+ * 
+ * Note: if you do have an actor with an owning Player Controller use the local player input subsystem instead.
+ */
+UCLASS()
+class ENHANCEDINPUT_API UEnhancedInputWorldSubsystem : public UTickableWorldSubsystem, public IEnhancedInputSubsystemInterface
+{
+	GENERATED_BODY()
+
+public:
+
+	//~ Begin USubsystem interface
+	virtual void Initialize(FSubsystemCollectionBase& Collection) override;
+	virtual void Deinitialize() override;
+	virtual bool ShouldCreateSubsystem(UObject* Outer) const override;
+	//~ End USubsystem interface
+
+	//~ Begin FTickableGameObject interface
+	//virtual UWorld* GetTickableGameObjectWorld() const override;
+	virtual bool IsTickableInEditor() const { return true; }
+	virtual ETickableTickType GetTickableTickType() const override;
+	virtual void Tick(float DeltaTime) override;
+	TStatId GetStatId() const override { RETURN_QUICK_DECLARE_CYCLE_STAT(UEnhancedInputWorldSubsystem, STATGROUP_Tickables); }
+	//~ End FTickableGameObject interface
+	
+	//~ Begin IEnhancedInputSubsystemInterface
+	virtual UEnhancedPlayerInput* GetPlayerInput() const override;
+	//~ End IEnhancedInputSubsystemInterface
+
+	/** Pushes this input component onto the stack to be processed by this subsystem's tick function */
+	UFUNCTION(BlueprintCallable, Category = "Input|World")
+	void PushInputComponent(AActor* Actor);
+
+	/** Removes this input component onto the stack to be processed by this subsystem's tick function */
+	UFUNCTION(BlueprintCallable, Category = "Input|World")
+	bool PopInputComponent(AActor* Actor);
+	
+	/** Start the consumption of input messages in this subsystem. This is required to have any Input Action delegates be fired. */
+	UFUNCTION(BlueprintCallable, Category = "Input|World")
+	void StartConsumingInput();
+
+	/** Tells this subsystem to stop ticking and consuming any input. This will stop any Input Action Delegates from being called. */
+	UFUNCTION(BlueprintCallable, Category = "Input|World")
+	void StopConsumingInput();
+	
+	/** Returns true if this subsystem is currently consuming input */
+    UFUNCTION(BlueprintCallable, BlueprintPure, Category = "Input|World")
+    bool IsConsumingInput() const { return bIsCurrentlyConsumingInput; }
+	
+    /** Inputs a key on this subsystem's player input which can then be processed as normal during Tick. */
+    bool InputKey(const FInputKeyParams& Params);
+	
+	/** Adds all the default mapping contexts */
+	void AddDefaultMappingContexts();
+
+	/** Removes all the default mapping contexts */
+	void RemoveDefaultMappingContexts();
+	
+private:
+
+	/** The player input that is processing the input within this subsystem */
+	UPROPERTY()
+	TObjectPtr<UEnhancedPlayerInput> PlayerInput = nullptr;
+
+	/**
+	 * Input processor that is created on Initalize.
+	 */
+	TSharedPtr<FEnhancedInputWorldProcessor> InputPreprocessor = nullptr;
+	
+	/** If true, then this subsystem will Tick and process input delegates. */
+	bool bIsCurrentlyConsumingInput = false;
+	
+	/** Internal. This is the current stack of InputComponents that is being processed by the PlayerInput. */
+	UPROPERTY(Transient)
+	TArray<TWeakObjectPtr<UInputComponent>> CurrentInputStack;
+};
diff --git a/Engine/Plugins/EnhancedInput/Source/EnhancedInput/Public/EnhancedInputWorldProcessor.h b/Engine/Plugins/EnhancedInput/Source/EnhancedInput/Public/EnhancedInputWorldProcessor.h
new file mode 100644
--- /dev/null
+++ b/Engine/Plugins/EnhancedInput/Source/EnhancedInput/Public/EnhancedInputWorldProcessor.h
@@ -0,0 +1,61 @@
+// Copyright Epic Games, Inc. All Rights Reserved.
+
+#pragma once
+
+#include "CoreMinimal.h"
+#include "Framework/Application/IInputProcessor.h"
+
+struct FInputKeyParams;
+class UWorld;
+
+/**
+ * The Enhanced Input World Processor is used to pass InputKey events to the
+ * Enhanced Input World Subsystem.
+ *
+ * It will not steal input, and all the functions here will return "False"
+ * so that other Input Processors still run with all their normal considerations. 
+ */
+class FEnhancedInputWorldProcessor : public IInputProcessor
+{
+public:
+	//~ Begin IInputProcessor interface
+	virtual void Tick(const float DeltaTime, FSlateApplication& SlateApp, TSharedRef<ICursor> Cursor) override;
+	virtual bool HandleKeyDownEvent(FSlateApplication& SlateApp, const FKeyEvent& InKeyEvent) override;
+	virtual bool HandleKeyUpEvent(FSlateApplication& SlateApp, const FKeyEvent& InKeyEvent) override;
+	virtual bool HandleAnalogInputEvent(FSlateApplication& SlateApp, const FAnalogInputEvent& InAnalogInputEvent) override;
+	virtual bool HandleMouseMoveEvent(FSlateApplication& SlateApp, const FPointerEvent& MouseEvent) override;
+	virtual bool HandleMouseButtonDownEvent( FSlateApplication& SlateApp, const FPointerEvent& MouseEvent) override;
+	virtual bool HandleMouseButtonUpEvent( FSlateApplication& SlateApp, const FPointerEvent& MouseEvent) override;
+	virtual bool HandleMouseButtonDoubleClickEvent(FSlateApplication& SlateApp, const FPointerEvent& MouseEvent) override;
+	virtual bool HandleMouseWheelOrGestureEvent(FSlateApplication& SlateApp, const FPointerEvent& InWheelEvent, const FPointerEvent* InGestureEvent) override;
+	virtual const TCHAR* GetDebugName() const { return TEXT("EnhancedInputWorldPreprocessor"); }
+	//~ End IInputProcessor interface
+
+	// Setter for the the world variable
+	void SetSubsystemWorld(UWorld* World) { WorldFromSubsystem = World; }
+	// Getter for the the world variable
+	UWorld* GetSubsystemWorld() const { return WorldFromSubsystem; }
+
+protected:
+
+	void UpdateCachedPointerPosition(const FPointerEvent& MouseEvent);
+
+	/** Called on Tick to process any accumulated mouse input */
+	void ProcessAccumulatedPointerInput(float DeltaTime);
+	
+	/** Tells the Enhanced Input World subsystem about the given key event */
+	bool InputKeyToSubsystem(const FInputKeyParams& Params);
+
+	/**
+	 * The number of cursor samples that happened this frame. The X and Y will be incremented when there is a mouse move event this frame
+	 * and it will be reset on ProcessAccumulatedPointerInput.
+	 */
+	FIntPoint NumCursorSamplesThisFrame = FIntPoint::ZeroValue;
+	
+	/** The value of the mouse cursor from the most recent mouse event */
+	FVector2D CachedCursorDelta = FVector2D::ZeroVector;
+
+private:
+	UWorld* WorldFromSubsystem = nullptr;
+
+};
\ No newline at end of file
diff --git a/Engine/Source/Runtime/Engine/Classes/GameFramework/Actor.h b/Engine/Source/Runtime/Engine/Classes/GameFramework/Actor.h
--- a/Engine/Source/Runtime/Engine/Classes/GameFramework/Actor.h
+++ b/Engine/Source/Runtime/Engine/Classes/GameFramework/Actor.h
@@ -1173,6 +1173,13 @@ public:
 	UFUNCTION(BlueprintCallable, Category="Input")
 	virtual void EnableInput(class APlayerController* PlayerController);
 
+	/** 
+	 * Creates an input component from the input component passed in
+	 * @param InputComponentToCreate The UInputComponent to create.
+	 */
+	UFUNCTION(BlueprintCallable, Category="Input")
+	virtual void CreateInputComponent(TSubclassOf<UInputComponent> InputComponentToCreate);
+
 	/** 
 	 * Removes this actor from the stack of input being handled by a PlayerController.
 	 * @param PlayerController The PlayerController whose input events we no longer want to receive. If null, this actor will stop receiving input from all PlayerControllers.
diff --git a/Engine/Source/Runtime/Engine/Private/Actor.cpp b/Engine/Source/Runtime/Engine/Private/Actor.cpp
--- a/Engine/Source/Runtime/Engine/Private/Actor.cpp
+++ b/Engine/Source/Runtime/Engine/Private/Actor.cpp
@@ -4219,6 +4219,19 @@ void AActor::EnableInput(APlayerController* PlayerController)
 	}
 }
 
+void AActor::CreateInputComponent(TSubclassOf<UInputComponent> InputComponentToCreate)
+{
+	if (InputComponentToCreate && !InputComponent)
+	{
+		InputComponent = NewObject<UInputComponent>(this, InputComponentToCreate);
+		InputComponent->RegisterComponent();
+		InputComponent->bBlockInput = bBlockInput;
+		InputComponent->Priority = InputPriority;
+
+		UInputDelegateBinding::BindInputDelegatesWithSubojects(this, InputComponent);
+	}
+}
+
 void AActor::DisableInput(APlayerController* PlayerController)
 {
 	if (InputComponent)
