Fixing typo in tooltip in Enhanced Input

#rnx

diff --git a/Engine/Plugins/EnhancedInput/Source/EnhancedInput/Public/EnhancedActionKeyMapping.h b/Engine/Plugins/EnhancedInput/Source/EnhancedInput/Public/EnhancedActionKeyMapping.h
--- a/Engine/Plugins/EnhancedInput/Source/EnhancedInput/Public/EnhancedActionKeyMapping.h
+++ b/Engine/Plugins/EnhancedInput/Source/EnhancedInput/Public/EnhancedActionKeyMapping.h
@@ -174,7 +174,7 @@ protected:
 	UPROPERTY(EditAnywhere, Instanced, Category = "Input|Settings", meta = (EditCondition = "SettingBehavior == EPlayerMappableKeySettingBehaviors::OverrideSettings", DisplayAfter = "SettingBehavior"))
 	TObjectPtr<UPlayerMappableKeySettings> PlayerMappableKeySettings = nullptr;
 
-	/** If true than this ActionKeyMapping will be exposed as a player mappable key */
+	/** If true then this ActionKeyMapping will be exposed as a player mappable key */
 	UPROPERTY(EditAnywhere, BlueprintReadOnly, Category = "Input|PlayerMappable")
 	uint8 bIsPlayerMappable : 1;
 
