Moved some of the logic from GatherClassDetailsCDOs to CustomizeDetails so the expensive operations that could be shared were done once.

Before: 359ms
After: 181ms

diff --git a/Engine/Plugins/EnhancedInput/Source/InputEditor/Private/InputCustomizations.cpp b/Engine/Plugins/EnhancedInput/Source/InputEditor/Private/InputCustomizations.cpp
--- a/Engine/Plugins/EnhancedInput/Source/InputEditor/Private/InputCustomizations.cpp
+++ b/Engine/Plugins/EnhancedInput/Source/InputEditor/Private/InputCustomizations.cpp
@@ -272,8 +272,46 @@ void FEnhancedInputDeveloperSettingsCustomization::CustomizeDetails(IDetailLayou
 	static const FName TriggerCategoryName = TEXT("Trigger Default Values");
 	static const FName ModifierCategoryName = TEXT("Modifier Default Values");
 
-	TArray<UObject*> ModifierCDOs = GatherClassDetailsCDOs(UInputModifier::StaticClass());
-	TArray<UObject*> TriggerCDOs = GatherClassDetailsCDOs(UInputTrigger::StaticClass());
+	FAssetRegistryModule& AssetRegistryModule = FModuleManager::LoadModuleChecked<FAssetRegistryModule>(FName("AssetRegistry"));
+	IAssetRegistry& AssetRegistry = AssetRegistryModule.Get();
+
+	TArray<UObject*> ModifierCDOs;
+	TArray<UObject*> TriggerCDOs;
+	// Search BPs via asset registry
+	{
+		FARFilter Filter;
+		Filter.ClassPaths.Add(UBlueprint::StaticClass()->GetClassPathName());
+		Filter.bRecursiveClasses = true;
+		TArray<FAssetData> BlueprintAssetData;
+		AssetRegistry.GetAssets(Filter, BlueprintAssetData);
+
+		FName NativeParentClass(TEXT("NativeParentClass"));
+		for (FAssetData& Asset : BlueprintAssetData)
+		{
+			FAssetDataTagMapSharedView::FFindTagResult Result = Asset.TagsAndValues.FindTag(NativeParentClass);
+			if (Result.IsSet() && !ExcludedAssetNames.Contains(Asset.AssetName))
+			{
+				const FString ClassObjectPath = FPackageName::ExportTextPathToObjectPath(Result.GetValue());
+				if (UClass* ParentClass = FindObjectSafe<UClass>(nullptr, *ClassObjectPath, true))
+				{
+					// TODO: Forcibly loading these assets could cause problems on projects with a large number of them.
+					if(ParentClass->IsChildOf(UInputModifier::StaticClass()))
+					{
+						UBlueprint* BP = CastChecked<UBlueprint>(Asset.GetAsset());
+						ModifierCDOs.AddUnique(BP->GeneratedClass->GetDefaultObject());
+					}
+					if (ParentClass->IsChildOf(UInputTrigger::StaticClass()))
+					{
+						UBlueprint* BP = CastChecked<UBlueprint>(Asset.GetAsset());
+						TriggerCDOs.AddUnique(BP->GeneratedClass->GetDefaultObject());
+					}
+				}
+			}
+		}
+	}
+
+	GatherNativeClassDetailsCDOs(UInputModifier::StaticClass(), ModifierCDOs);
+	GatherNativeClassDetailsCDOs(UInputTrigger::StaticClass(), TriggerCDOs);
 	ExcludedAssetNames.Reset();
 	
 	// Add The modifier/trigger defaults that are generated via CDO to the details builder
@@ -281,8 +319,6 @@ void FEnhancedInputDeveloperSettingsCustomization::CustomizeDetails(IDetailLayou
 	CustomizeCDOValues(DetailBuilder, TriggerCategoryName, TriggerCDOs);
 	
 	// Support for updating blueprint based triggers and modifiers in the settings panel
-	FAssetRegistryModule& AssetRegistryModule = FModuleManager::LoadModuleChecked<FAssetRegistryModule>(FName("AssetRegistry"));
-	IAssetRegistry& AssetRegistry = AssetRegistryModule.Get();
 	if (!AssetRegistry.OnAssetAdded().IsBoundToObject(this))
 	{
 		AssetRegistry.OnAssetAdded().AddRaw(this, &FEnhancedInputDeveloperSettingsCustomization::OnAssetAdded);
@@ -325,10 +361,8 @@ void FEnhancedInputDeveloperSettingsCustomization::CustomizeDetails(const TShare
 	CustomizeDetails(*DetailBuilder);
 }
 
-TArray<UObject*> FEnhancedInputDeveloperSettingsCustomization::GatherClassDetailsCDOs(UClass* Class)
+void FEnhancedInputDeveloperSettingsCustomization::GatherNativeClassDetailsCDOs(UClass* Class, TArray<UObject*>& CDOs)
 {
-	TArray<UObject*> CDOs;
-
 	// Search native classes
 	for (TObjectIterator<UClass> ClassIt; ClassIt; ++ClassIt)
 	{
@@ -346,34 +380,6 @@ TArray<UObject*> FEnhancedInputDeveloperSettingsCustomization::GatherClassDetail
 		CDOs.AddUnique(ClassIt->GetDefaultObject());
 	}
 
-	// Search BPs via asset registry
-	FAssetRegistryModule& AssetRegistryModule = FModuleManager::LoadModuleChecked<FAssetRegistryModule>(FName("AssetRegistry"));
-	IAssetRegistry& AssetRegistry = AssetRegistryModule.Get();
-
-	FARFilter Filter;
-	Filter.ClassPaths.Add(UBlueprint::StaticClass()->GetClassPathName());
-	Filter.bRecursiveClasses = true;
-	TArray<FAssetData> BlueprintAssetData;
-	AssetRegistry.GetAssets(Filter, BlueprintAssetData);
-
-	for (FAssetData& Asset : BlueprintAssetData)
-	{
-		FAssetDataTagMapSharedView::FFindTagResult Result = Asset.TagsAndValues.FindTag(TEXT("NativeParentClass"));
-		if (Result.IsSet() && !ExcludedAssetNames.Contains(Asset.AssetName))
-		{
-			const FString ClassObjectPath = FPackageName::ExportTextPathToObjectPath(Result.GetValue());
-			if (UClass* ParentClass = FindObjectSafe<UClass>(nullptr, *ClassObjectPath, true))
-			{
-				if (ParentClass->IsChildOf(Class))
-				{
-					// TODO: Forcibly loading these assets could cause problems on projects with a large number of them.
-					UBlueprint* BP = CastChecked<UBlueprint>(Asset.GetAsset());
-					CDOs.AddUnique(BP->GeneratedClass->GetDefaultObject());
-				}
-			}
-		}
-	}
-
 	// Strip objects with no config stored properties
 	CDOs.RemoveAll([Class](UObject* Object) {
 		UClass* ObjectClass = Object->GetClass();
@@ -396,8 +402,6 @@ TArray<UObject*> FEnhancedInputDeveloperSettingsCustomization::GatherClassDetail
 		}
 		return true;
 	});
-
-	return CDOs;
 }
 
 bool FEnhancedInputDeveloperSettingsCustomization::DoesClassHaveSubtypes(UClass* Class)
diff --git a/Engine/Plugins/EnhancedInput/Source/InputEditor/Private/InputCustomizations.h b/Engine/Plugins/EnhancedInput/Source/InputEditor/Private/InputCustomizations.h
--- a/Engine/Plugins/EnhancedInput/Source/InputEditor/Private/InputCustomizations.h
+++ b/Engine/Plugins/EnhancedInput/Source/InputEditor/Private/InputCustomizations.h
@@ -75,8 +75,8 @@ public:
 
 private:
 
-	/** Gather all of the CDO's for the given class, Native and Blueprint. */
-	static TArray<UObject*> GatherClassDetailsCDOs(UClass* Class);
+	/** Gather all of the native class CDOs for the given class. */
+	static void GatherNativeClassDetailsCDOs(UClass* Class, TArray<UObject*>& CDOs);
 
 	/**
 	 * Called when any Asset is added, removed, or renamed.
