PR #9845: EnhancedInput : K2Node_EnhancedInputAction fix for UInputAction child classes (Contributed by [REDACTED_PERSON])

diff --git a/Engine/Plugins/EnhancedInput/Source/InputBlueprintNodes/Private/K2Node_EnhancedInputAction.cpp b/Engine/Plugins/EnhancedInput/Source/InputBlueprintNodes/Private/K2Node_EnhancedInputAction.cpp
--- a/Engine/Plugins/EnhancedInput/Source/InputBlueprintNodes/Private/K2Node_EnhancedInputAction.cpp
+++ b/Engine/Plugins/EnhancedInput/Source/InputBlueprintNodes/Private/K2Node_EnhancedInputAction.cpp
@@ -420,7 +420,7 @@ void UK2Node_EnhancedInputAction::GetMenuActions(FBlueprintActionDatabaseRegistr
 		}
 
 		TArray<FAssetData> ActionAssets;
-		AssetRegistry.GetAssetsByClass(UInputAction::StaticClass()->GetClassPathName(), ActionAssets);
+		AssetRegistry.GetAssetsByClass(UInputAction::StaticClass()->GetClassPathName(), ActionAssets, true);
 		for (const FAssetData& ActionAsset : ActionAssets)
 		{
 			UBlueprintNodeSpawner* NodeSpawner = UBlueprintNodeSpawner::Create(GetClass());
