Add auto complete for "show debug enhancedinput" and "showdebug devices"

#rnx

diff --git a/Engine/Plugins/EnhancedInput/Config/Input.ini b/Engine/Plugins/EnhancedInput/Config/Input.ini
new file mode 100644
--- /dev/null
+++ b/Engine/Plugins/EnhancedInput/Config/Input.ini
@@ -0,0 +1,3 @@
+[/Script/EngineSettings.ConsoleSettings]
++ManualAutoCompleteList=(Command="ShowDebug EnhancedInput",Desc="Displays debug information about the current state of any Enhanced Input Mapping Contexts")
++ManualAutoCompleteList=(Command="ShowDebug Devices",Desc="Displays all current Platform Users and their Input Devices.")
