Fix a possible crash when you use a "Get action value" node in a component.
I have been unable to reproduce this locally, but the user's callstack indicates clearly that the Actor or action is somehow null. I am adding these logs and the ensure to hopefully get some further info about the crash in the future, if it does repro for anyone else, and prevent the crash.

diff --git a/Engine/Plugins/EnhancedInput/Source/EnhancedInput/Private/EnhancedInputLibrary.cpp b/Engine/Plugins/EnhancedInput/Source/EnhancedInput/Private/EnhancedInputLibrary.cpp
--- a/Engine/Plugins/EnhancedInput/Source/EnhancedInput/Private/EnhancedInputLibrary.cpp
+++ b/Engine/Plugins/EnhancedInput/Source/EnhancedInput/Private/EnhancedInputLibrary.cpp
@@ -43,8 +43,25 @@ void UEnhancedInputLibrary::RequestRebuildControlMappingsUsingContext(const UInp
 
 FInputActionValue UEnhancedInputLibrary::GetBoundActionValue(AActor* Actor, const UInputAction* Action)
 {
-	UEnhancedInputComponent* EIC = Cast<UEnhancedInputComponent>(Actor->InputComponent);
-	return EIC ? EIC->GetBoundActionValue(Action) : FInputActionValue(Action->ValueType, FVector::ZeroVector);
+	if (Actor && Action)
+	{
+		UEnhancedInputComponent* EIC = Cast<UEnhancedInputComponent>(Actor->InputComponent);
+		return EIC ? EIC->GetBoundActionValue(Action) : FInputActionValue(Action->ValueType, FVector::ZeroVector);
+	}
+	else
+	{
+		if (!Actor)
+		{
+			UE_LOG(LogEnhancedInput, Error, TEXT("UEnhancedInputLibrary::GetBoundActionValue was called with an invalid Actor!"));
+		}
+		
+		if (!Action)
+		{
+			UE_LOG(LogEnhancedInput, Error, TEXT("UEnhancedInputLibrary::GetBoundActionValue was called with an invalid Action!"));
+		}		
+		ensureMsgf(false, TEXT("Invalid GetBoundActionValue call. Check logs for details!"));
+	}
+	return FInputActionValue();
 }
 
 
