Clamping Input Modifier dead zones between 0 and 1

diff --git a/Engine/Plugins/EnhancedInput/Source/EnhancedInput/Public/InputModifiers.h b/Engine/Plugins/EnhancedInput/Source/EnhancedInput/Public/InputModifiers.h
--- a/Engine/Plugins/EnhancedInput/Source/EnhancedInput/Public/InputModifiers.h
+++ b/Engine/Plugins/EnhancedInput/Source/EnhancedInput/Public/InputModifiers.h
@@ -74,11 +74,11 @@ class UInputModifierDeadZone : public UInputModifier
 public:
 
 	// Threshold below which input is ignored
-	UPROPERTY(EditInstanceOnly, BlueprintReadWrite, Category=Settings, Config)
+	UPROPERTY(EditInstanceOnly, BlueprintReadWrite, Category=Settings, Config, meta=(ClampMin=0, ClampMax=1))
 	float LowerThreshold = 0.2f;
 
 	// Threshold above which input is clamped to 1
-	UPROPERTY(EditInstanceOnly, BlueprintReadWrite, Category=Settings, Config)
+	UPROPERTY(EditInstanceOnly, BlueprintReadWrite, Category=Settings, Config, meta=(ClampMin=0, ClampMax=1))
 	float UpperThreshold = 1.f;
 
 	UPROPERTY(EditInstanceOnly, BlueprintReadWrite, Category=Settings, Config)
