Added Key Slots for player mapped keys. It is now possible to create Secondary, Tertiary,(..and so on) mappings from a single default mapping name.
Avoids having to create an empty "_Secondary" mapping in the mapping context.
Added section to ShowDebugInfo that adds mappings created through add player mappable keys.

diff --git a/Engine/Plugins/EnhancedInput/Source/EnhancedInput/Private/EnhancedInputSubsystemInterface.cpp b/Engine/Plugins/EnhancedInput/Source/EnhancedInput/Private/EnhancedInputSubsystemInterface.cpp
--- a/Engine/Plugins/EnhancedInput/Source/EnhancedInput/Private/EnhancedInputSubsystemInterface.cpp
+++ b/Engine/Plugins/EnhancedInput/Source/EnhancedInput/Private/EnhancedInputSubsystemInterface.cpp
@@ -405,14 +405,26 @@ TArray<FEnhancedActionKeyMapping> IEnhancedInputSubsystemInterface::GetAllPlayer
 	return PlayerMappableMappings;
 }
 
-int32 IEnhancedInputSubsystemInterface::AddPlayerMappedKey(const FName MappingName, const FKey NewKey, const FModifyContextOptions& Options)
+int32 IEnhancedInputSubsystemInterface::AddPlayerMappedKey(const FName MappingName, const FKey NewKey, const FModifyContextOptions& Options /*= FModifyContextOptions()*/)
+{
+	return AddPlayerMappedKeyInSlot(MappingName, NewKey, FPlayerMappableKeySlot::FirstKeySlot, Options);
+}
+
+int32 IEnhancedInputSubsystemInterface::K2_AddPlayerMappedKeyInSlot(const FName MappingName, const FKey NewKey, const FPlayerMappableKeySlot& KeySlot /*= FPlayerMappableKeySlot()*/, const FModifyContextOptions& Options /*= FModifyContextOptions()*/)
+{
+	return AddPlayerMappedKeyInSlot(MappingName, NewKey, KeySlot, Options);
+}
+
+int32 IEnhancedInputSubsystemInterface::AddPlayerMappedKeyInSlot(const FName MappingName, const FKey NewKey, const FPlayerMappableKeySlot& KeySlot /*= FPlayerMappableKeySlot::FirstKeySlot*/, const FModifyContextOptions& Options /*= FModifyContextOptions()*/)
 {
 	int32 NumMappingsApplied = 0;
 	if (MappingName != NAME_None)
 	{
 		if (UEnhancedPlayerInput* const PlayerInput = GetPlayerInput())
 		{
-			PlayerMappedSettings.Add(MappingName, NewKey);
+			TMap<FPlayerMappableKeySlot, FKey>& PlayerMappedKeySlots = PlayerMappedSettings.FindOrAdd(MappingName);
+			FKey& KeyInSlot = PlayerMappedKeySlots.FindOrAdd(KeySlot);
+			KeyInSlot = NewKey;
 			++NumMappingsApplied;
 		}
 
@@ -420,22 +432,48 @@ int32 IEnhancedInputSubsystemInterface::AddPlayerMappedKey(const FName MappingNa
 	}
 	else
 	{
-		UE_LOG(LogEnhancedInput, Warning, TEXT("Attempted to AddPlayerMappedKey with an invalid MappingName! Mapping has not been applied."));
+		UE_LOG(LogEnhancedInput, Warning, TEXT("Attempted to AddPlayerMappedKeyInSlot with an invalid MappingName! Mapping has not been applied."));
 	}
 	return NumMappingsApplied;
 }
 
-int32 IEnhancedInputSubsystemInterface::RemovePlayerMappedKey(const FName MappingName, const FModifyContextOptions& Options)
+int32 IEnhancedInputSubsystemInterface::RemovePlayerMappedKey(const FName MappingName, const FModifyContextOptions& Options /*= FModifyContextOptions()*/)
 {
-	int32 NumMappingsApplied = 0;
+	return RemovePlayerMappedKeyInSlot(MappingName, FPlayerMappableKeySlot::FirstKeySlot, Options);
+}
+
+int32 IEnhancedInputSubsystemInterface::K2_RemovePlayerMappedKeyInSlot(const FName MappingName, const FPlayerMappableKeySlot& KeySlot /*= FPlayerMappableKeySlot()*/, const FModifyContextOptions& Options /*= FModifyContextOptions()*/)
+{
+	return RemovePlayerMappedKeyInSlot(MappingName, KeySlot, Options);
+}
+
+int32 IEnhancedInputSubsystemInterface::RemovePlayerMappedKeyInSlot(const FName MappingName, const FPlayerMappableKeySlot& KeySlot /*= FPlayerMappableKeySlot::FirstKeySlot*/, const FModifyContextOptions& Options /*= FModifyContextOptions()*/)
+{
+	int32 NumMappingsRemoved = 0;
 	if (UEnhancedPlayerInput* const PlayerInput = GetPlayerInput())
 	{
-		NumMappingsApplied = PlayerMappedSettings.Remove(MappingName);
+		if (TMap<FPlayerMappableKeySlot, FKey>* PlayerMappedKeySlots = PlayerMappedSettings.Find(MappingName))
+		{
+			NumMappingsRemoved = PlayerMappedKeySlots->Remove(KeySlot);
+		}
 	}
 
 	RequestRebuildControlMappings(Options);
-	
-	return NumMappingsApplied;
+
+	return NumMappingsRemoved;
+}
+
+int32 IEnhancedInputSubsystemInterface::RemoveAllPlayerMappedKeysForMapping(const FName MappingName, const FModifyContextOptions& Options /*= FModifyContextOptions()*/)
+{
+	int32 NumMappingsRemoved = 0;
+	if (UEnhancedPlayerInput* const PlayerInput = GetPlayerInput())
+	{
+		NumMappingsRemoved = PlayerMappedSettings.Remove(MappingName);
+	}
+
+	RequestRebuildControlMappings(Options);
+
+	return NumMappingsRemoved;
 }
 
 void IEnhancedInputSubsystemInterface::RemoveAllPlayerMappedKeys(const FModifyContextOptions& Options)
@@ -450,14 +488,37 @@ void IEnhancedInputSubsystemInterface::RemoveAllPlayerMappedKeys(const FModifyCo
 
 FKey IEnhancedInputSubsystemInterface::GetPlayerMappedKey(const FName MappingName) const
 {
-	if (const FKey* MappedKey = PlayerMappedSettings.Find(MappingName))
+	return GetPlayerMappedKeyInSlot(MappingName, FPlayerMappableKeySlot::FirstKeySlot);
+}
+
+FKey IEnhancedInputSubsystemInterface::K2_GetPlayerMappedKeyInSlot(const FName MappingName, const FPlayerMappableKeySlot& KeySlot /*= FPlayerMappableKeySlot()*/) const
+{
+	return GetPlayerMappedKeyInSlot(MappingName, KeySlot);
+}
+
+FKey IEnhancedInputSubsystemInterface::GetPlayerMappedKeyInSlot(const FName MappingName, const FPlayerMappableKeySlot& KeySlot /*= FPlayerMappableKeySlot()*/) const
+{
+	if (const TMap<FPlayerMappableKeySlot, FKey>* PlayerMappedKeySlots = PlayerMappedSettings.Find(MappingName))
 	{
-		return FKey(*MappedKey);
+		if (const FKey* PlayerMappedKey = PlayerMappedKeySlots->Find(KeySlot))
+		{
+			return *PlayerMappedKey;
+		}
 	}
-	
 	return EKeys::Invalid;
 }
 
+TArray<FKey> IEnhancedInputSubsystemInterface::GetAllPlayerMappedKeys(const FName MappingName) const
+{
+	TArray<FKey> PlayerMappedKeys;
+	if (const TMap<FPlayerMappableKeySlot, FKey>* PlayerMappedKeySlots = PlayerMappedSettings.Find(MappingName))
+	{
+		PlayerMappedKeySlots->GenerateValueArray(PlayerMappedKeys);
+		return PlayerMappedKeys;
+	}
+	return PlayerMappedKeys;
+}
+
 void IEnhancedInputSubsystemInterface::AddPlayerMappableConfig(const UPlayerMappableInputConfig* Config, const FModifyContextOptions& Options)
 {
 	if(Config)
@@ -657,50 +718,60 @@ void IEnhancedInputSubsystemInterface::RebuildControlMappings()
 
 		for (FEnhancedActionKeyMapping& Mapping : OrderedMappings)
 		{
-			if (FKey* PlayerKey = PlayerMappedSettings.Find(Mapping.GetMappingName()))
+			TArray<FKey> PlayerMappedKeys = GetAllPlayerMappedKeys(Mapping.GetMappingName());
+			if (PlayerMappedKeys.IsEmpty())
 			{
-				Mapping.Key = *PlayerKey;
+				//If we didn't find any player mapped keys use the default key from the mapping.
+				PlayerMappedKeys.Add(Mapping.Key);
 			}
-			
-			if (Mapping.Action && !AppliedKeys.Contains(Mapping.Key))
+
+			for (const FKey& PlayerMappedKey : PlayerMappedKeys)
 			{
-				// TODO: Wasteful query as we've already established chord state within ReorderMappings. Store TOptional bConsumeInput per mapping, allowing override? Query override via delegate?
-				auto IsChord = [](const UInputTrigger* Trigger) { return Cast<const UInputTriggerChordAction>(Trigger) != nullptr; };
-				bool bHasActionChords = HasTriggerWith(IsChord, Mapping.Action->Triggers);
-				bool bHasChords = HasTriggerWith(IsChord, Mapping.Triggers) || bHasActionChords;
+				Mapping.Key = PlayerMappedKey;
 
-				// Chorded actions can't consume input or they would hide the action they are chording.
-				if (!bHasChords && Mapping.Action->bConsumeInput)
+				if (Mapping.Action && !AppliedKeys.Contains(Mapping.Key))
 				{
-					ContextAppliedKeys.Add(Mapping.Key);
-				}
+					// TODO: Wasteful query as we've already established chord state within ReorderMappings. Store TOptional bConsumeInput per mapping, allowing override? Query override via delegate?
+					auto IsChord = [](const UInputTrigger* Trigger)
+					{
+						return Cast<const UInputTriggerChordAction>(Trigger) != nullptr;
+					};
+					bool bHasActionChords = HasTriggerWith(IsChord, Mapping.Action->Triggers);
+					bool bHasChords = HasTriggerWith(IsChord, Mapping.Triggers) || bHasActionChords;
 
-				int32 NewMappingIndex = PlayerInput->AddMapping(Mapping);
-				FEnhancedActionKeyMapping& NewMapping = PlayerInput->EnhancedActionMappings[NewMappingIndex];
+					// Chorded actions can't consume input or they would hide the action they are chording.
+					if (!bHasChords && Mapping.Action->bConsumeInput)
+					{
+						ContextAppliedKeys.Add(Mapping.Key);
+					}
 
-				// Re-instance modifiers
-				DeepCopyPtrArray<UInputModifier>(Mapping.Modifiers, NewMapping.Modifiers);
+					int32 NewMappingIndex = PlayerInput->AddMapping(Mapping);
+					FEnhancedActionKeyMapping& NewMapping = PlayerInput->EnhancedActionMappings[NewMappingIndex];
 
-				ApplyAxisPropertyModifiers(PlayerInput, NewMapping);
+					// Re-instance modifiers
+					DeepCopyPtrArray<UInputModifier>(Mapping.Modifiers, NewMapping.Modifiers);
 
-				// Re-instance triggers
-				DeepCopyPtrArray<UInputTrigger>(Mapping.Triggers, NewMapping.Triggers);
+					ApplyAxisPropertyModifiers(PlayerInput, NewMapping);
 
-				if (bHasChords)
-				{
-					// TODO: Re-prioritize chorded mappings (within same context only?) by number of chorded actions, so Ctrl + Alt + [key] > Ctrl + [key] > [key].
-					// TODO: Above example shouldn't block [key] if only Alt is down, as there is no direct Alt + [key] mapping.y
-					ChordedMappings.Add(NewMappingIndex);
+					// Re-instance triggers
+					DeepCopyPtrArray<UInputTrigger>(Mapping.Triggers, NewMapping.Triggers);
 
-					// Action level chording triggers need to be evaluated at the mapping level to ensure they block early enough.
-					// TODO: Continuing to evaluate these at the action level is redundant.
-					if (bHasActionChords)
+					if (bHasChords)
 					{
-						for (const UInputTrigger* Trigger : Mapping.Action->Triggers)
+						// TODO: Re-prioritize chorded mappings (within same context only?) by number of chorded actions, so Ctrl + Alt + [key] > Ctrl + [key] > [key].
+						// TODO: Above example shouldn't block [key] if only Alt is down, as there is no direct Alt + [key] mapping.y
+						ChordedMappings.Add(NewMappingIndex);
+
+						// Action level chording triggers need to be evaluated at the mapping level to ensure they block early enough.
+						// TODO: Continuing to evaluate these at the action level is redundant.
+						if (bHasActionChords)
 						{
-							if (IsChord(Trigger))
+							for (const UInputTrigger* Trigger : Mapping.Action->Triggers)
 							{
-								NewMapping.Triggers.Add(DuplicateObject(Trigger, nullptr));
+								if (IsChord(Trigger))
+								{
+									NewMapping.Triggers.Add(DuplicateObject(Trigger, nullptr));
+								}
 							}
 						}
 					}
diff --git a/Engine/Plugins/EnhancedInput/Source/EnhancedInput/Private/EnhancedInputSubsystemInterfaceDebug.cpp b/Engine/Plugins/EnhancedInput/Source/EnhancedInput/Private/EnhancedInputSubsystemInterfaceDebug.cpp
--- a/Engine/Plugins/EnhancedInput/Source/EnhancedInput/Private/EnhancedInputSubsystemInterfaceDebug.cpp
+++ b/Engine/Plugins/EnhancedInput/Source/EnhancedInput/Private/EnhancedInputSubsystemInterfaceDebug.cpp
@@ -162,6 +162,16 @@ void IEnhancedInputSubsystemInterface::ShowDebugInfo(UCanvas* Canvas)
 				}
 			}
 
+			const TArray<FEnhancedActionKeyMapping>& EnhancedActionMappings = PlayerInput->GetEnhancedActionMappings();
+			for (const FEnhancedActionKeyMapping& EnhancedActionMapping : EnhancedActionMappings)
+			{
+				if (TArray<FEnhancedActionKeyMapping>* Mappings = ActionMappings.Find(EnhancedActionMapping.Action))
+				{
+					//Add any mapping that might have been added as player mappable keys.
+					Mappings->AddUnique(EnhancedActionMapping);
+				}
+			}
+
 			Sort(OrderedActions.GetData(), OrderedActions.Num(), [](const UInputAction& A, const UInputAction& B) { return A.GetFName().LexicalLess(B.GetFName()); });
 
 			for (const UInputAction* Action : OrderedActions)
diff --git a/Engine/Plugins/EnhancedInput/Source/EnhancedInput/Private/PlayerMappableKeySlot.cpp b/Engine/Plugins/EnhancedInput/Source/EnhancedInput/Private/PlayerMappableKeySlot.cpp
new file mode 100644
--- /dev/null
+++ b/Engine/Plugins/EnhancedInput/Source/EnhancedInput/Private/PlayerMappableKeySlot.cpp
@@ -0,0 +1,21 @@
+// Copyright Epic Games, Inc. All Rights Reserved.
+
+#include "PlayerMappableKeySlot.h"
+
+FPlayerMappableKeySlot FPlayerMappableKeySlot::FirstKeySlot = FPlayerMappableKeySlot(0);
+FPlayerMappableKeySlot FPlayerMappableKeySlot::SecondKeySlot = FPlayerMappableKeySlot(1);
+FPlayerMappableKeySlot FPlayerMappableKeySlot::ThirdKeySlot = FPlayerMappableKeySlot(2);
+FPlayerMappableKeySlot FPlayerMappableKeySlot::FourthKeySlot = FPlayerMappableKeySlot(3);
+
+FPlayerMappableKeySlot::FPlayerMappableKeySlot() { }
+FPlayerMappableKeySlot::FPlayerMappableKeySlot(const int32 InSlotNumber) : SlotNumber(InSlotNumber) { }
+
+int32 FPlayerMappableKeySlot::GetSlotNumber() const
+{
+	return SlotNumber;
+}
+
+uint32 GetTypeHash(const FPlayerMappableKeySlot& Ref)
+{
+	return GetTypeHash(Ref.GetSlotNumber());
+}
diff --git a/Engine/Plugins/EnhancedInput/Source/EnhancedInput/Public/EnhancedInputLibrary.h b/Engine/Plugins/EnhancedInput/Source/EnhancedInput/Public/EnhancedInputLibrary.h
--- a/Engine/Plugins/EnhancedInput/Source/EnhancedInput/Public/EnhancedInputLibrary.h
+++ b/Engine/Plugins/EnhancedInput/Source/EnhancedInput/Public/EnhancedInputLibrary.h
@@ -7,6 +7,7 @@
 #include "InputActionValue.h"
 #include "InputMappingQuery.h"
 #include "Kismet/BlueprintFunctionLibrary.h"
+#include "PlayerMappableKeySlot.h"
 #include "UObject/ObjectMacros.h"
 
 #include "EnhancedInputLibrary.generated.h"
@@ -97,4 +98,16 @@ public:
 	/** Converts a FInputActionValue to a string */
 	UFUNCTION(BlueprintPure, meta = (DisplayName = "To String (InputActionValue)", CompactNodeTitle = "->", BlueprintAutocast), Category = "Utilities|String")
 	static FString Conv_InputActionValueToString(FInputActionValue ActionValue);
+
+	UFUNCTION(BlueprintPure, Category = "Input", meta = (ReturnDisplayName = "First Player Mappable Key Slot"))
+	static FPlayerMappableKeySlot& GetFirstPlayerMappableKeySlot() { return FPlayerMappableKeySlot::FirstKeySlot; };
+
+	UFUNCTION(BlueprintPure, Category = "Input", meta = (ReturnDisplayName = "Second Player Mappable Key Slot"))
+	static FPlayerMappableKeySlot& GetSecondPlayerMappableKeySlot() { return FPlayerMappableKeySlot::SecondKeySlot; };
+
+	UFUNCTION(BlueprintPure, Category = "Input", meta = (ReturnDisplayName = "Third Player Mappable Key Slot"))
+	static FPlayerMappableKeySlot& GetThirdPlayerMappableKeySlot() { return FPlayerMappableKeySlot::ThirdKeySlot; };
+
+	UFUNCTION(BlueprintPure, Category = "Input", meta = (ReturnDisplayName = "Fourth Player Mappable Key Slot"))
+	static FPlayerMappableKeySlot& GetFourthPlayerMappableKeySlot() { return FPlayerMappableKeySlot::FourthKeySlot; };
 };
diff --git a/Engine/Plugins/EnhancedInput/Source/EnhancedInput/Public/EnhancedInputSubsystemInterface.h b/Engine/Plugins/EnhancedInput/Source/EnhancedInput/Public/EnhancedInputSubsystemInterface.h
--- a/Engine/Plugins/EnhancedInput/Source/EnhancedInput/Public/EnhancedInputSubsystemInterface.h
+++ b/Engine/Plugins/EnhancedInput/Source/EnhancedInput/Public/EnhancedInputSubsystemInterface.h
@@ -8,6 +8,7 @@
 #include "UObject/Interface.h"
 #include "EnhancedActionKeyMapping.h"
 #include "EnhancedPlayerInput.h"
+#include "PlayerMappableKeySlot.h"
 
 #include "EnhancedInputSubsystemInterface.generated.h"
 
@@ -180,31 +181,103 @@ public:
 	virtual TArray<FEnhancedActionKeyMapping> GetAllPlayerMappableActionKeyMappings() const;
 	
 	/**
-	 * Replace any currently applied mappings to this key mapping with the given new one.
+	 * Emplace or replace any currently applied key in the first key slot for mapping of MappingName.
 	 * Requests a rebuild of the player mappings. 
 	 *
 	 * @return The number of mappings that have been replaced
 	 */
+	//UE_DEPRECATED(5.2, "AddPlayerMappedKey has been deprecated, please use AddPlayerMappedKeyInSlot instead.")
 	UFUNCTION(BlueprintCallable, BlueprintCosmetic, Category = "Input|PlayerMappable", meta=(AutoCreateRefTerm = "Options"))
 	virtual int32 AddPlayerMappedKey(const FName MappingName, const FKey NewKey, const FModifyContextOptions& Options = FModifyContextOptions());
 
 	/**
-	 * Remove any player mappings with to the given action
+	 * Emplace or replace any currently applied key in KeySlot for mapping of MappingName.
+	 * Requests a rebuild of the player mappings.
+	 *
+	 * @return The number of mappings that have been replaced
+	 */
+	UFUNCTION(BlueprintCallable, BlueprintCosmetic, Category = "Input|PlayerMappable", meta = (AutoCreateRefTerm = "KeySlot, Options"))
+	virtual int32 K2_AddPlayerMappedKeyInSlot(const FName MappingName, const FKey NewKey, const FPlayerMappableKeySlot& KeySlot = FPlayerMappableKeySlot(), const FModifyContextOptions& Options = FModifyContextOptions());
+
+	/**
+	 * Emplace or replace any currently applied key in KeySlot for mapping of MappingName.
+	 * Requests a rebuild of the player mappings.
+	 *
+	 * @return The number of mappings that have been replaced
+	 */
+	virtual int32 AddPlayerMappedKeyInSlot(const FName MappingName, const FKey NewKey, const FPlayerMappableKeySlot& KeySlot = FPlayerMappableKeySlot::FirstKeySlot, const FModifyContextOptions& Options = FModifyContextOptions());
+
+	/**
+	 * Removes player mapped key in the first KeySlot for mapping of MappingName.
 	 * Requests a rebuild of the player mappings. 
 	 *
 	 * @return The number of mappings that have been removed
 	 */
+	//UE_DEPRECATED(5.2, "RemovePlayerMappedKey has been deprecated, please use RemovePlayerMappedKeyInSlot instead.")
 	UFUNCTION(BlueprintCallable, BlueprintCosmetic, Category = "Input|PlayerMappable", meta=(AutoCreateRefTerm = "Options"))
 	virtual int32 RemovePlayerMappedKey(const FName MappingName, const FModifyContextOptions& Options = FModifyContextOptions());
 
 	/**
-	 * Get the player mapped key to the given mapping name. If there is not a player mapped key, then this will return
+	 * Removes player mapped key in the KeySlot for mapping of MappingName.
+	 * Requests a rebuild of the player mappings.
+	 *
+	 * @return The number of mappings that have been removed
+	 */
+	UFUNCTION(BlueprintCallable, BlueprintCosmetic, Category = "Input|PlayerMappable", meta = (AutoCreateRefTerm = "KeySlot, Options"))
+	virtual int32 K2_RemovePlayerMappedKeyInSlot(const FName MappingName, const FPlayerMappableKeySlot& KeySlot = FPlayerMappableKeySlot(), const FModifyContextOptions& Options = FModifyContextOptions());
+
+	/**
+	 * Removes player mapped key in the KeySlot for mapping of MappingName.
+	 * Requests a rebuild of the player mappings.
+	 *
+	 * @return The number of mappings that have been removed
+	 */
+	virtual int32 RemovePlayerMappedKeyInSlot(const FName MappingName, const FPlayerMappableKeySlot& KeySlot = FPlayerMappableKeySlot::FirstKeySlot, const FModifyContextOptions& Options = FModifyContextOptions());
+
+	/**
+	 * Removes all player mapped keys for mapping of MappingName.
+	 * Requests a rebuild of the player mappings.
+	 *
+	 * @return The number of mappings that have been removed
+	 */
+	UFUNCTION(BlueprintCallable, BlueprintCosmetic, Category = "Input|PlayerMappable", meta = (AutoCreateRefTerm = "KeySlot, Options"))
+	virtual int32 RemoveAllPlayerMappedKeysForMapping(const FName MappingName, const FModifyContextOptions& Options = FModifyContextOptions());
+
+	/**
+	 * Get the player mapped key in first slot to the given mapping name. If there is not a player mapped key, then this will return
 	 * EKeys::Invalid.
 	 *
 	 * @param MappingName	The FName of the mapped key that would have been set with the AddPlayerMappedKey function.
 	 */
+	//UE_DEPRECATED(5.2, "GetPlayerMappedKey has been deprecated, please use GetPlayerMappedKeyInSlot instead.")
 	UFUNCTION(BlueprintCallable, BlueprintCosmetic, Category = "Input|PlayerMappable")
 	virtual FKey GetPlayerMappedKey(const FName MappingName) const;
+
+	/**
+	 * Get the player mapped key in first slot to the given mapping name. If there is not a player mapped key, then this will return
+	 * EKeys::Invalid.
+	 *
+	 * @param MappingName	The FName of the mapped key that would have been set with the AddPlayerMappedKey function.
+	 */
+	UFUNCTION(BlueprintCallable, BlueprintCosmetic, Category = "Input|PlayerMappable", meta = (AutoCreateRefTerm = "KeySlot"))
+	virtual FKey K2_GetPlayerMappedKeyInSlot(const FName MappingName, const FPlayerMappableKeySlot& KeySlot = FPlayerMappableKeySlot()) const;
+
+	/**
+	 * Get the player mapped key in first slot to the given mapping name. If there is not a player mapped key, then this will return
+	 * EKeys::Invalid.
+	 *
+	 * @param MappingName	The FName of the mapped key that would have been set with the AddPlayerMappedKey function.
+	 */
+	virtual FKey GetPlayerMappedKeyInSlot(const FName MappingName, const FPlayerMappableKeySlot& KeySlot = FPlayerMappableKeySlot::FirstKeySlot) const;
+
+	/**
+	 * Get all the player mapped keys to the given mapping name. If there is not a player mapped key, then this will return
+	 * TArray<FKey>().
+	 *
+	 * @param MappingName	The FName of the mapped key that would have been set with the AddPlayerMappedKey function.
+	 */
+	UFUNCTION(BlueprintCallable, BlueprintCosmetic, Category = "Input|PlayerMappable")
+	virtual TArray<FKey> GetAllPlayerMappedKeys(const FName MappingName) const;
 	
 	/**
 	 * Remove All PlayerMappedKeys
@@ -273,8 +346,8 @@ private:
 	TMap<TWeakObjectPtr<const UInputAction>, FInputActionValue> ForcedActions;
 	TMap<FKey, FInputActionValue> ForcedKeys;
 
-	/** A map of any player mapped keys to the key that they should redirect to instead */
-	TMap<FName, FKey> PlayerMappedSettings;
+	/** A map of any player mapped keys to the keys that they should redirect to instead */
+	TMap<FName, TMap<FPlayerMappableKeySlot, FKey>> PlayerMappedSettings;
 
 	/**
 	 * A map of input actions with a Chorded trigger, mapped to the action they are dependent on.
diff --git a/Engine/Plugins/EnhancedInput/Source/EnhancedInput/Public/PlayerMappableKeySlot.h b/Engine/Plugins/EnhancedInput/Source/EnhancedInput/Public/PlayerMappableKeySlot.h
new file mode 100644
--- /dev/null
+++ b/Engine/Plugins/EnhancedInput/Source/EnhancedInput/Public/PlayerMappableKeySlot.h
@@ -0,0 +1,37 @@
+// Copyright Epic Games, Inc. All Rights Reserved.
+
+#pragma once
+
+#include "PlayerMappableKeySlot.generated.h"
+
+/**
+* Explicitly identifies the slot for a player mappable key
+*/
+USTRUCT(BlueprintType)
+struct ENHANCEDINPUT_API FPlayerMappableKeySlot
+{
+	GENERATED_BODY()
+
+public:
+
+	FPlayerMappableKeySlot();
+	FPlayerMappableKeySlot(int32 InSlotNumber);
+
+	virtual int32 GetSlotNumber() const;
+
+	UPROPERTY(EditAnywhere, BlueprintReadWrite, Category = "Input")
+	int32 SlotNumber = 0;
+
+	static FPlayerMappableKeySlot FirstKeySlot;
+	static FPlayerMappableKeySlot SecondKeySlot;
+	static FPlayerMappableKeySlot ThirdKeySlot;
+	static FPlayerMappableKeySlot FourthKeySlot;
+
+	bool operator==(const FPlayerMappableKeySlot& OtherKeySlot) const
+	{
+		return GetSlotNumber() == OtherKeySlot.GetSlotNumber();
+	}
+
+};
+
+ENHANCEDINPUT_API uint32 GetTypeHash(const FPlayerMappableKeySlot& Ref);
