Fixed CIS issue. Added missing virtual Destructor.

diff --git a/Engine/Plugins/EnhancedInput/Source/EnhancedInput/Public/PlayerMappableKeySlot.h b/Engine/Plugins/EnhancedInput/Source/EnhancedInput/Public/PlayerMappableKeySlot.h
--- a/Engine/Plugins/EnhancedInput/Source/EnhancedInput/Public/PlayerMappableKeySlot.h
+++ b/Engine/Plugins/EnhancedInput/Source/EnhancedInput/Public/PlayerMappableKeySlot.h
@@ -16,6 +16,7 @@ public:
 
 	FPlayerMappableKeySlot();
 	FPlayerMappableKeySlot(int32 InSlotNumber);
+	virtual ~FPlayerMappableKeySlot();
 
 	virtual int32 GetSlotNumber() const;
 
