Fixed CIS issue. Added missing virtual Destructor cpp.

diff --git a/Engine/Plugins/EnhancedInput/Source/EnhancedInput/Private/PlayerMappableKeySlot.cpp b/Engine/Plugins/EnhancedInput/Source/EnhancedInput/Private/PlayerMappableKeySlot.cpp
--- a/Engine/Plugins/EnhancedInput/Source/EnhancedInput/Private/PlayerMappableKeySlot.cpp
+++ b/Engine/Plugins/EnhancedInput/Source/EnhancedInput/Private/PlayerMappableKeySlot.cpp
@@ -9,6 +9,7 @@ FPlayerMappableKeySlot FPlayerMappableKeySlot::FourthKeySlot = FPlayerMappableKe
 
 FPlayerMappableKeySlot::FPlayerMappableKeySlot() { }
 FPlayerMappableKeySlot::FPlayerMappableKeySlot(const int32 InSlotNumber) : SlotNumber(InSlotNumber) { }
+FPlayerMappableKeySlot::~FPlayerMappableKeySlot() { }
 
 int32 FPlayerMappableKeySlot::GetSlotNumber() const
 {
