Unshelved from pending changelist '23430190':

Complete the deprecation of Add player mapped key functions in favor of the newer "Player Mappable Key Slot" struct.

#pf 6390fd8a1776b8c21c30d309
#rnx

diff --git a/Engine/Plugins/EnhancedInput/Config/BaseEnhancedInput.ini b/Engine/Plugins/EnhancedInput/Config/BaseEnhancedInput.ini
--- a/Engine/Plugins/EnhancedInput/Config/BaseEnhancedInput.ini
+++ b/Engine/Plugins/EnhancedInput/Config/BaseEnhancedInput.ini
@@ -1,2 +1,7 @@
 [CoreRedirects]
-+FunctionRedirects=(OldName="/Script/InputMappingContext.UnmapAction", NewName="/Script/InputMappingContext.UnmapAllKeysFromAction")
\ No newline at end of file
++FunctionRedirects=(OldName="/Script/InputMappingContext.UnmapAction", NewName="/Script/InputMappingContext.UnmapAllKeysFromAction")
+
+; 5.2
++FunctionRedirects=(OldName="/Script/EnhancedInput.EnhancedInputSubsystemInterface.AddPlayerMappedKey", NewName="/Script/EnhancedInput.EnhancedInputSubsystemInterface.K2_AddPlayerMappedKeyInSlot")
++FunctionRedirects=(OldName="/Script/EnhancedInput.EnhancedInputSubsystemInterface.RemovePlayerMappedKey", NewName="/Script/EnhancedInput.EnhancedInputSubsystemInterface.K2_RemovePlayerMappedKeyInSlot")
++FunctionRedirects=(OldName="/Script/EnhancedInput.EnhancedInputSubsystemInterface.GetPlayerMappedKey", NewName="/Script/EnhancedInput.EnhancedInputSubsystemInterface.K2_GetPlayerMappedKeyInSlot")
\ No newline at end of file
diff --git a/Engine/Plugins/EnhancedInput/Source/EnhancedInput/Public/EnhancedInputSubsystemInterface.h b/Engine/Plugins/EnhancedInput/Source/EnhancedInput/Public/EnhancedInputSubsystemInterface.h
--- a/Engine/Plugins/EnhancedInput/Source/EnhancedInput/Public/EnhancedInputSubsystemInterface.h
+++ b/Engine/Plugins/EnhancedInput/Source/EnhancedInput/Public/EnhancedInputSubsystemInterface.h
@@ -186,8 +186,7 @@ public:
 	 *
 	 * @return The number of mappings that have been replaced
 	 */
-	//UE_DEPRECATED(5.2, "AddPlayerMappedKey has been deprecated, please use AddPlayerMappedKeyInSlot instead.")
-	UFUNCTION(BlueprintCallable, BlueprintCosmetic, Category = "Input|PlayerMappable", meta=(AutoCreateRefTerm = "Options"))
+	UE_DEPRECATED(5.2, "AddPlayerMappedKey has been deprecated, please use AddPlayerMappedKeyInSlot instead.")
 	virtual int32 AddPlayerMappedKey(const FName MappingName, const FKey NewKey, const FModifyContextOptions& Options = FModifyContextOptions());
 
 	/**
@@ -196,7 +195,7 @@ public:
 	 *
 	 * @return The number of mappings that have been replaced
 	 */
-	UFUNCTION(BlueprintCallable, BlueprintCosmetic, Category = "Input|PlayerMappable", meta = (AutoCreateRefTerm = "KeySlot, Options"))
+	UFUNCTION(BlueprintCallable, BlueprintCosmetic, Category = "Input|PlayerMappable", meta = (DisplayName = "Add Player Mapped Key In Slot", AutoCreateRefTerm = "KeySlot, Options"))
 	virtual int32 K2_AddPlayerMappedKeyInSlot(const FName MappingName, const FKey NewKey, const FPlayerMappableKeySlot& KeySlot = FPlayerMappableKeySlot(), const FModifyContextOptions& Options = FModifyContextOptions());
 
 	/**
@@ -213,8 +212,7 @@ public:
 	 *
 	 * @return The number of mappings that have been removed
 	 */
-	//UE_DEPRECATED(5.2, "RemovePlayerMappedKey has been deprecated, please use RemovePlayerMappedKeyInSlot instead.")
-	UFUNCTION(BlueprintCallable, BlueprintCosmetic, Category = "Input|PlayerMappable", meta=(AutoCreateRefTerm = "Options"))
+	UE_DEPRECATED(5.2, "RemovePlayerMappedKey has been deprecated, please use RemovePlayerMappedKeyInSlot instead.")
 	virtual int32 RemovePlayerMappedKey(const FName MappingName, const FModifyContextOptions& Options = FModifyContextOptions());
 
 	/**
@@ -223,7 +221,7 @@ public:
 	 *
 	 * @return The number of mappings that have been removed
 	 */
-	UFUNCTION(BlueprintCallable, BlueprintCosmetic, Category = "Input|PlayerMappable", meta = (AutoCreateRefTerm = "KeySlot, Options"))
+	UFUNCTION(BlueprintCallable, BlueprintCosmetic, Category = "Input|PlayerMappable", meta = (DisplayName="Remove Player Mapped Key In Slot", AutoCreateRefTerm = "KeySlot, Options"))
 	virtual int32 K2_RemovePlayerMappedKeyInSlot(const FName MappingName, const FPlayerMappableKeySlot& KeySlot = FPlayerMappableKeySlot(), const FModifyContextOptions& Options = FModifyContextOptions());
 
 	/**
@@ -249,8 +247,7 @@ public:
 	 *
 	 * @param MappingName	The FName of the mapped key that would have been set with the AddPlayerMappedKey function.
 	 */
-	//UE_DEPRECATED(5.2, "GetPlayerMappedKey has been deprecated, please use GetPlayerMappedKeyInSlot instead.")
-	UFUNCTION(BlueprintCallable, BlueprintCosmetic, Category = "Input|PlayerMappable")
+	UE_DEPRECATED(5.2, "GetPlayerMappedKey has been deprecated, please use GetPlayerMappedKeyInSlot instead.")
 	virtual FKey GetPlayerMappedKey(const FName MappingName) const;
 
 	/**
@@ -259,7 +256,7 @@ public:
 	 *
 	 * @param MappingName	The FName of the mapped key that would have been set with the AddPlayerMappedKey function.
 	 */
-	UFUNCTION(BlueprintCallable, BlueprintCosmetic, Category = "Input|PlayerMappable", meta = (AutoCreateRefTerm = "KeySlot"))
+	UFUNCTION(BlueprintCallable, BlueprintCosmetic, Category = "Input|PlayerMappable", meta = (DisplayName="Get Player Mapped Key In Slot",AutoCreateRefTerm = "KeySlot"))
 	virtual FKey K2_GetPlayerMappedKeyInSlot(const FName MappingName, const FPlayerMappableKeySlot& KeySlot = FPlayerMappableKeySlot()) const;
 
 	/**
diff --git a/Engine/Plugins/Experimental/VirtualCamera/Source/VCamCore/Private/VCamComponent.cpp b/Engine/Plugins/Experimental/VirtualCamera/Source/VCamCore/Private/VCamComponent.cpp
--- a/Engine/Plugins/Experimental/VirtualCamera/Source/VCamCore/Private/VCamComponent.cpp
+++ b/Engine/Plugins/Experimental/VirtualCamera/Source/VCamCore/Private/VCamComponent.cpp
@@ -617,7 +617,7 @@ bool UVCamComponent::AddInputProfileWithCurrentlyActiveMappings(const FName Prof
 					const FName MappingName = Mapping.GetMappingName();
 
 					// Prefer to use the current mapped key but fallback to the default if no key is mapped
-					FKey CurrentKey = EnhancedInputSubsystemInterface->GetPlayerMappedKey(MappingName);
+					FKey CurrentKey = EnhancedInputSubsystemInterface->GetPlayerMappedKeyInSlot(MappingName);
 					if (!CurrentKey.IsValid())
 					{
 						CurrentKey = Mapping.Key;
@@ -666,7 +666,7 @@ FKey UVCamComponent::GetPlayerMappedKey(const FName MappingName) const
 {
 	if (const IEnhancedInputSubsystemInterface* EnhancedInputSubsystemInterface = GetEnhancedInputSubsystemInterface())
 	{
-		return EnhancedInputSubsystemInterface->GetPlayerMappedKey(MappingName);
+		return EnhancedInputSubsystemInterface->GetPlayerMappedKeyInSlot(MappingName);
 	}
 	return EKeys::Invalid;
 }
@@ -1211,7 +1211,7 @@ void UVCamComponent::ApplyInputProfile()
 			// Ensure we have a valid name to map
 			if (MappingName != NAME_None)
 			{
-				EnhancedInputSubsystemInterface->AddPlayerMappedKey(MappingName, NewKey);
+				EnhancedInputSubsystemInterface->AddPlayerMappedKeyInSlot(MappingName, NewKey);
 			}
 		}
 	}
