Adding Enhanced Input world subsystem debug info

diff --git a/Engine/Plugins/EnhancedInput/Config/Input.ini b/Engine/Plugins/EnhancedInput/Config/Input.ini
--- a/Engine/Plugins/EnhancedInput/Config/Input.ini
+++ b/Engine/Plugins/EnhancedInput/Config/Input.ini
@@ -1,2 +1,3 @@
 [/Script/EngineSettings.ConsoleSettings]
 +ManualAutoCompleteList=(Command="ShowDebug EnhancedInput",Desc="Displays debug information about the current state of any Enhanced Input Mapping Contexts")
++ManualAutoCompleteList=(Command="ShowDebug WorldSubsystemInput",Desc="Displays debug information about the current state of any Enhanced Input Mapping Contexts applied to the Enhanced Input world subsystem.")
diff --git a/Engine/Plugins/EnhancedInput/Source/EnhancedInput/Private/EnhancedInputModule.cpp b/Engine/Plugins/EnhancedInput/Source/EnhancedInput/Private/EnhancedInputModule.cpp
--- a/Engine/Plugins/EnhancedInput/Source/EnhancedInput/Private/EnhancedInputModule.cpp
+++ b/Engine/Plugins/EnhancedInput/Source/EnhancedInput/Private/EnhancedInputModule.cpp
@@ -306,6 +306,8 @@ void FEnhancedInputModule::Tick(float DeltaTime)
 void FEnhancedInputModule::OnShowDebugInfo(AHUD* HUD, UCanvas* Canvas, const FDebugDisplayInfo& DisplayInfo, float& YL, float& YPos)
 {
 	static const FName NAME_EnhancedInput("EnhancedInput");
+	static const FName NAME_WorldSubsystemInput("WorldSubsystemInput");
+
 	if (Canvas)
 	{
 		if (HUD->ShouldDisplayDebug(NAME_EnhancedInput))
@@ -324,6 +326,14 @@ void FEnhancedInputModule::OnShowDebugInfo(AHUD* HUD, UCanvas* Canvas, const FDe
 				FirstPlayer->ShowDebugInfo(Canvas);
 			}
 		}
+
+		if (HUD->ShouldDisplayDebug(NAME_WorldSubsystemInput))
+		{
+			for (TObjectIterator<UEnhancedInputWorldSubsystem> It; It; ++It)
+			{
+				It->ShowDebugInfo(Canvas);
+			}
+		}
 	}
 	else
 	{
diff --git a/Engine/Plugins/EnhancedInput/Source/EnhancedInput/Private/EnhancedInputSubsystemInterfaceDebug.cpp b/Engine/Plugins/EnhancedInput/Source/EnhancedInput/Private/EnhancedInputSubsystemInterfaceDebug.cpp
--- a/Engine/Plugins/EnhancedInput/Source/EnhancedInput/Private/EnhancedInputSubsystemInterfaceDebug.cpp
+++ b/Engine/Plugins/EnhancedInput/Source/EnhancedInput/Private/EnhancedInputSubsystemInterfaceDebug.cpp
@@ -54,8 +54,7 @@ struct FVisualizationTexture
 };
 static TMap<uint64, FVisualizationTexture> CachedModifierVisualizations;
 
-
-void IEnhancedInputSubsystemInterface::ShowDebugInfo(UCanvas* Canvas)
+void IEnhancedInputSubsystemInterface::ShowMappingContextDebugInfo(UCanvas* Canvas, const UEnhancedPlayerInput* PlayerInput)
 {
 	if (!Canvas)
 	{
@@ -64,37 +63,17 @@ void IEnhancedInputSubsystemInterface::ShowDebugInfo(UCanvas* Canvas)
 
 	FDisplayDebugManager& DisplayDebugManager = Canvas->DisplayDebugManager;
 
-	// TODO: Localize some/all debug output?
-	const UEnhancedPlayerInput* PlayerInput = GetPlayerInput();
 	if (!PlayerInput)
 	{
-		DisplayDebugManager.SetDrawColor(FColor::Orange);
-		DisplayDebugManager.DrawString(TEXT("This player does not support Enhanced Input. To enable it update Project Settings -> Input -> Default Classes to the Enhanced versions."));
+		DisplayDebugManager.SetDrawColor(FColor::Red);
+		DisplayDebugManager.DrawString(TEXT("Invalid Player Input!"));
 		return;
 	}
 
-	if (APlayerController* PC = Cast<APlayerController>(PlayerInput->GetOuter()))
-	{
-		DisplayDebugManager.SetDrawColor(FColor::White);
-		DisplayDebugManager.DrawString(FString::Printf(TEXT("Player: %s"), *PC->GetFName().ToString()));
-
-		// TODO: Display input stack? Remove input stack?
-		//TArray<UInputComponent*> InputStack;
-		//PC->BuildInputStack(InputStack);
-		//FString InputStackStr;
-		//for(UInputComponent* IC : InputStack)
-		//{
-		//	AActor* Owner = InputStack[i]->GetOwner();
-		//	InputStackStr += Owner ? Owner->GetFName().ToString() + "." : "" + IC->GetFName().ToString() + " > ";
-		//}
-		//DisplayDebugManager.SetDrawColor(FColor::White);
-		//DisplayDebugManager.DrawString(FString::Printf(TEXT("Input stack: %s"), *InputStackStr.LeftChop(3)));
-	}
-
 	if (PlayerInput->EnhancedActionMappings.Num() + PlayerInput->LastInjectedActions.Num() == 0)
 	{
 		DisplayDebugManager.SetDrawColor(FColor::Orange);
-		DisplayDebugManager.DrawString(TEXT("No enhanced player input action mappings have been applied to this player."));
+		DisplayDebugManager.DrawString(TEXT("No enhanced player input action mappings have been applied to this input."));
 	}
 	else
 	{
@@ -283,6 +262,45 @@ void IEnhancedInputSubsystemInterface::ShowDebugInfo(UCanvas* Canvas)
 	}
 }
 
+void IEnhancedInputSubsystemInterface::ShowDebugInfo(UCanvas* Canvas)
+{
+	if (!Canvas)
+	{
+		return;
+	}
+
+	FDisplayDebugManager& DisplayDebugManager = Canvas->DisplayDebugManager;
+
+	// TODO: Localize some/all debug output?
+	UEnhancedPlayerInput* PlayerInput = GetPlayerInput();
+	if (!PlayerInput)
+	{
+		DisplayDebugManager.SetDrawColor(FColor::Orange);
+		DisplayDebugManager.DrawString(TEXT("This player does not support Enhanced Input. To enable it update Project Settings -> Input -> Default Classes to the Enhanced versions."));
+		return;
+	}
+
+	if (APlayerController* PC = Cast<APlayerController>(PlayerInput->GetOuter()))
+	{
+		DisplayDebugManager.SetDrawColor(FColor::White);
+		DisplayDebugManager.DrawString(FString::Printf(TEXT("Player: %s"), *PC->GetFName().ToString()));
+
+		// TODO: Display input stack? Remove input stack?
+		//TArray<UInputComponent*> InputStack;
+		//PC->BuildInputStack(InputStack);
+		//FString InputStackStr;
+		//for(UInputComponent* IC : InputStack)
+		//{
+		//	AActor* Owner = InputStack[i]->GetOwner();
+		//	InputStackStr += Owner ? Owner->GetFName().ToString() + "." : "" + IC->GetFName().ToString() + " > ";
+		//}
+		//DisplayDebugManager.SetDrawColor(FColor::White);
+		//DisplayDebugManager.DrawString(FString::Printf(TEXT("Input stack: %s"), *InputStackStr.LeftChop(3)));
+	}
+
+	ShowMappingContextDebugInfo(Canvas, PlayerInput);
+}
+
 void IEnhancedInputSubsystemInterface::ShowDebugActionModifiers(UCanvas* Canvas, const UInputAction* Action)
 {
 	FDisplayDebugManager& DisplayDebugManager = Canvas->DisplayDebugManager;
diff --git a/Engine/Plugins/EnhancedInput/Source/EnhancedInput/Private/EnhancedInputSubsystems.cpp b/Engine/Plugins/EnhancedInput/Source/EnhancedInput/Private/EnhancedInputSubsystems.cpp
--- a/Engine/Plugins/EnhancedInput/Source/EnhancedInput/Private/EnhancedInputSubsystems.cpp
+++ b/Engine/Plugins/EnhancedInput/Source/EnhancedInput/Private/EnhancedInputSubsystems.cpp
@@ -7,6 +7,7 @@
 #include "EnhancedInputDeveloperSettings.h"
 #include "EnhancedInputModule.h"
 #include "EnhancedPlayerInput.h"
+#include "Engine/Canvas.h"
 #include "Engine/World.h"
 #include "Framework/Application/SlateApplication.h"
 #include "GameFramework/PlayerController.h"
@@ -264,3 +265,26 @@ void UEnhancedInputWorldSubsystem::RemoveDefaultMappingContexts()
 		}
 	}
 }
+
+void UEnhancedInputWorldSubsystem::ShowDebugInfo(UCanvas* Canvas)
+{
+	if (!Canvas)
+	{
+		return;
+	}
+
+	FDisplayDebugManager& DisplayDebugManager = Canvas->DisplayDebugManager;
+
+	UEnhancedPlayerInput* WorldSubsystemPlayerInput = GetPlayerInput();
+	if (!WorldSubsystemPlayerInput)
+	{
+		DisplayDebugManager.SetDrawColor(FColor::Orange);
+		DisplayDebugManager.DrawString(TEXT("This player does not support Enhanced Input."));
+		return;
+	}
+	
+	DisplayDebugManager.SetDrawColor(FColor::White);
+	DisplayDebugManager.DrawString(FString::Printf(TEXT("World Subsystem from %s"), LexToString(GetWorld()->WorldType)));
+	
+	ShowMappingContextDebugInfo(Canvas, WorldSubsystemPlayerInput);
+}
diff --git a/Engine/Plugins/EnhancedInput/Source/EnhancedInput/Public/EnhancedInputSubsystemInterface.h b/Engine/Plugins/EnhancedInput/Source/EnhancedInput/Public/EnhancedInputSubsystemInterface.h
--- a/Engine/Plugins/EnhancedInput/Source/EnhancedInput/Public/EnhancedInputSubsystemInterface.h
+++ b/Engine/Plugins/EnhancedInput/Source/EnhancedInput/Public/EnhancedInputSubsystemInterface.h
@@ -13,12 +13,14 @@
 #include "EnhancedInputSubsystemInterface.generated.h"
 
 class APlayerController;
+class UCanvas;
 class UInputMappingContext;
 class UInputAction;
 class UEnhancedPlayerInput;
 class UInputModifier;
 class UInputTrigger;
 class UPlayerMappableInputConfig;
+class UWorldSubsystem;
 
 // Subsystem interface
 UINTERFACE(MinimalAPI, meta=(CannotImplementInterfaceInBlueprint))
@@ -49,7 +51,7 @@ struct FModifyContextOptions
 		, bForceImmediately(false)
 	{}
 
-	// If true than any keys that are pressed during the rebuild of control mappings will be ignored until they are released.
+	// If true then any keys that are pressed during the rebuild of control mappings will be ignored until they are released.
 	UPROPERTY(EditAnywhere, BlueprintReadWrite, Category = "Input")
 	uint8 bIgnoreAllPressedKeysUntilRelease : 1;
 
@@ -312,6 +314,9 @@ protected:
 
 	/** Function that will be called when Control Mappings have been rebuilt this tick. */
 	virtual void ControlMappingsRebuiltThisFrame() {}
+
+	// helper function to display debug about mapping context info
+	void ShowMappingContextDebugInfo(UCanvas* Canvas, const UEnhancedPlayerInput* PlayerInput);
 	
 private:
 
@@ -376,8 +381,8 @@ private:
 	 */
 	bool bControlMappingsRebuiltThisTick = false;
 
-	// Debug visualization implemented in EnhancedInputSubsystemsDebug.cpp
-	void ShowDebugInfo(class UCanvas* Canvas);
+	// Debug visualization for enhanced input local player subsystem
+	virtual void ShowDebugInfo(UCanvas* Canvas);
 	
 	void ShowDebugActionModifiers(UCanvas* Canvas, const UInputAction* Action);
 	static void PurgeDebugVisualizations();
diff --git a/Engine/Plugins/EnhancedInput/Source/EnhancedInput/Public/EnhancedInputSubsystems.h b/Engine/Plugins/EnhancedInput/Source/EnhancedInput/Public/EnhancedInputSubsystems.h
--- a/Engine/Plugins/EnhancedInput/Source/EnhancedInput/Public/EnhancedInputSubsystems.h
+++ b/Engine/Plugins/EnhancedInput/Source/EnhancedInput/Public/EnhancedInputSubsystems.h
@@ -98,6 +98,7 @@ public:
 	/** Removes all the default mapping contexts */
 	void RemoveDefaultMappingContexts();
 	
+	virtual void ShowDebugInfo(UCanvas* Canvas) override;
 private:
 
 	/** The player input that is processing the input within this subsystem */
