Passing the OwningActionKeyMapping to the PlayerMappedSettings to let child classes know in what context they are being called.

diff --git a/Engine/Plugins/EnhancedInput/Source/EnhancedInput/Private/EnhancedActionKeyMapping.cpp b/Engine/Plugins/EnhancedInput/Source/EnhancedInput/Private/EnhancedActionKeyMapping.cpp
--- a/Engine/Plugins/EnhancedInput/Source/EnhancedInput/Private/EnhancedActionKeyMapping.cpp
+++ b/Engine/Plugins/EnhancedInput/Source/EnhancedInput/Private/EnhancedActionKeyMapping.cpp
@@ -29,7 +29,7 @@ ENHANCEDINPUT_API FName FEnhancedActionKeyMapping::GetMappingName() const
 	{
 		if (UPlayerMappableKeySettings* MappableKeySettings = GetPlayerMappableKeySettings())
 		{
-			return MappableKeySettings->GetMappingName();
+			return MappableKeySettings->MakeMappingName(this);
 		}
 		return PlayerMappableOptions.Name;
 	}
diff --git a/Engine/Plugins/EnhancedInput/Source/EnhancedInput/Public/PlayerMappableKeySettings.h b/Engine/Plugins/EnhancedInput/Source/EnhancedInput/Public/PlayerMappableKeySettings.h
--- a/Engine/Plugins/EnhancedInput/Source/EnhancedInput/Public/PlayerMappableKeySettings.h
+++ b/Engine/Plugins/EnhancedInput/Source/EnhancedInput/Public/PlayerMappableKeySettings.h
@@ -8,6 +8,8 @@
 #include "UObject/ObjectPtr.h"
 #include "PlayerMappableKeySettings.generated.h"
 
+struct FEnhancedActionKeyMapping;
+
 /**
 * Hold setting information of an Action Input or a Action Key Mapping for setting screen and save purposes.
 */
@@ -18,6 +20,7 @@ class ENHANCEDINPUT_API UPlayerMappableKeySettings : public UObject
 
 public:
 
+	virtual FName MakeMappingName(const FEnhancedActionKeyMapping* OwningActionKeyMapping) const { return GetMappingName(); }
 	virtual FName GetMappingName() const { return Name; }
 
 #if WITH_EDITOR
