Making Enhanced Input World Subsystem Project Settings grouped in the same category

diff --git a/Engine/Plugins/EnhancedInput/Source/EnhancedInput/Private/EnhancedInputSubsystems.cpp b/Engine/Plugins/EnhancedInput/Source/EnhancedInput/Private/EnhancedInputSubsystems.cpp
--- a/Engine/Plugins/EnhancedInput/Source/EnhancedInput/Private/EnhancedInputSubsystems.cpp
+++ b/Engine/Plugins/EnhancedInput/Source/EnhancedInput/Private/EnhancedInputSubsystems.cpp
@@ -224,7 +224,7 @@ bool UEnhancedInputWorldSubsystem::InputKey(const FInputKeyParams& Params)
 		{
 			if (GetDefault<UEnhancedInputDeveloperSettings>()->bShouldLogAllWorldSubsystemInputs)
 			{
-				UE_LOG(LogWorldSubsystemInput, Log, TEXT("EI World Subsystem InputKey : [%s]"), *Params.Key.ToString());
+				UE_LOG(LogWorldSubsystemInput, Log, TEXT("EI %s World Subsystem InputKey : [%s]"), LexToString(GetWorld()->WorldType), *Params.Key.ToString());
 			}
 			return PlayerInput->InputKey(Params);
 		}
diff --git a/Engine/Plugins/EnhancedInput/Source/EnhancedInput/Public/EnhancedInputDeveloperSettings.h b/Engine/Plugins/EnhancedInput/Source/EnhancedInput/Public/EnhancedInputDeveloperSettings.h
--- a/Engine/Plugins/EnhancedInput/Source/EnhancedInput/Public/EnhancedInputDeveloperSettings.h
+++ b/Engine/Plugins/EnhancedInput/Source/EnhancedInput/Public/EnhancedInputDeveloperSettings.h
@@ -75,6 +75,6 @@ public:
  	 * If true then the Enhanced Input world subsystem will log all input that is being processed by it (keypresses, analog values, etc)
  	 * Note: This can produce A LOT of logs, so only use this if you are debugging something.
  	 */
-	UPROPERTY(config, EditAnywhere, Category = "Enhanced Input World Subsystem", meta=(ConsoleVariable="EnhancedInput.bShouldLogAllWorldSubsystemInputs"))
+	UPROPERTY(config, EditAnywhere, Category = "Enhanced Input", meta=(ConsoleVariable="EnhancedInput.bShouldLogAllWorldSubsystemInputs"))
 	uint8 bShouldLogAllWorldSubsystemInputs : 1;
 };
\ No newline at end of file
