Unshelved from pending changelist '23593668':

Add an option to disable the Enhanced Input world subsystem. Set EnhancedInput.bEnableWorldSubsystem to false to turn it off.

#pf skip

diff --git a/Engine/Plugins/EnhancedInput/Source/EnhancedInput/Private/EnhancedInputSubsystems.cpp b/Engine/Plugins/EnhancedInput/Source/EnhancedInput/Private/EnhancedInputSubsystems.cpp
--- a/Engine/Plugins/EnhancedInput/Source/EnhancedInput/Private/EnhancedInputSubsystems.cpp
+++ b/Engine/Plugins/EnhancedInput/Source/EnhancedInput/Private/EnhancedInputSubsystems.cpp
@@ -96,7 +96,14 @@ bool UEnhancedInputWorldSubsystem::ShouldCreateSubsystem(UObject* Outer) const
 		return false;
 	}
 	
-	return Super::ShouldCreateSubsystem(Outer);
+	// Getting setting on whether to turn off subsystem or not
+	const bool bShouldCreate = GetDefault<UEnhancedInputDeveloperSettings>()->bEnableWorldSubsystem;
+	if (!bShouldCreate)
+	{
+		UE_LOG(LogWorldSubsystemInput, Log, TEXT("UEnhancedInputDeveloperSettings::bEnableWorldSubsystem is false, the world subsystem will not be created!"));
+	}
+	
+	return bShouldCreate && Super::ShouldCreateSubsystem(Outer);
 }
 
 ETickableTickType UEnhancedInputWorldSubsystem::GetTickableTickType() const
diff --git a/Engine/Plugins/EnhancedInput/Source/EnhancedInput/Public/EnhancedInputDeveloperSettings.h b/Engine/Plugins/EnhancedInput/Source/EnhancedInput/Public/EnhancedInputDeveloperSettings.h
--- a/Engine/Plugins/EnhancedInput/Source/EnhancedInput/Public/EnhancedInputDeveloperSettings.h
+++ b/Engine/Plugins/EnhancedInput/Source/EnhancedInput/Public/EnhancedInputDeveloperSettings.h
@@ -71,6 +71,10 @@ public:
 	UPROPERTY(EditAnywhere, Category = "Enhanced Input", meta=(ConsoleVariable="EnhancedInput.OnlyTriggerLastActionInChord"))
 	bool bShouldOnlyTriggerLastActionInChord = true;
 	
+	/** If true, then the world subsystem will be created. */
+	UPROPERTY(config, EditAnywhere, Category = "Enhanced Input")
+	bool bEnableWorldSubsystem = true;
+	
 	/**
  	 * If true then the Enhanced Input world subsystem will log all input that is being processed by it (keypresses, analog values, etc)
  	 * Note: This can produce A LOT of logs, so only use this if you are debugging something.
