Initalize player mappable configs ConfigName property to the Asset Name by default instead of leaving it as None

#rnx
#pf skip

diff --git a/Engine/Plugins/EnhancedInput/Source/EnhancedInput/Private/PlayerMappableInputConfig.cpp b/Engine/Plugins/EnhancedInput/Source/EnhancedInput/Private/PlayerMappableInputConfig.cpp
--- a/Engine/Plugins/EnhancedInput/Source/EnhancedInput/Private/PlayerMappableInputConfig.cpp
+++ b/Engine/Plugins/EnhancedInput/Source/EnhancedInput/Private/PlayerMappableInputConfig.cpp
@@ -10,6 +10,8 @@
 
 UPlayerMappableInputConfig::UPlayerMappableInputConfig(const FObjectInitializer& ObjectInitializer)
 	: Super(ObjectInitializer)
+	, ConfigName(GetFName())
+	, ConfigDisplayName(FText::FromName(ConfigName))
 {
 	ResetToDefault();
 }
diff --git a/Engine/Plugins/EnhancedInput/Source/EnhancedInput/Public/PlayerMappableInputConfig.h b/Engine/Plugins/EnhancedInput/Source/EnhancedInput/Public/PlayerMappableInputConfig.h
--- a/Engine/Plugins/EnhancedInput/Source/EnhancedInput/Public/PlayerMappableInputConfig.h
+++ b/Engine/Plugins/EnhancedInput/Source/EnhancedInput/Public/PlayerMappableInputConfig.h
@@ -54,11 +54,11 @@ protected:
 	
 	/** The unique name of this config that can be used when saving it */
 	UPROPERTY(EditAnywhere, BlueprintReadOnly, Category = "Input|PlayerMappable")
-	FName ConfigName = NAME_None;
+	FName ConfigName;
 
 	/** The display name that can be used  */
 	UPROPERTY(EditAnywhere, BlueprintReadOnly, Category = "Input|PlayerMappable")
-	FText ConfigDisplayName = FText::GetEmpty();
+	FText ConfigDisplayName;
 
 	/** A flag that can be used to mark this Input Config as deprecated to your player/designers. */
 	UPROPERTY(EditAnywhere, BlueprintReadOnly, Category = "Input|PlayerMappable")
