Fixing action value type toast showing up when modifying triggers

diff --git a/Engine/Plugins/EnhancedInput/Source/EnhancedInput/Private/InputAction.cpp b/Engine/Plugins/EnhancedInput/Source/EnhancedInput/Private/InputAction.cpp
--- a/Engine/Plugins/EnhancedInput/Source/EnhancedInput/Private/InputAction.cpp
+++ b/Engine/Plugins/EnhancedInput/Source/EnhancedInput/Private/InputAction.cpp
@@ -107,6 +107,7 @@ void FTriggerStateTracker::SetStateForNoTriggers(ETriggerState State)
 // Record input action property changes for later processing
 
 TSet<const UInputAction*> UInputAction::ActionsWithModifiedValueTypes;
+TSet<const UInputAction*> UInputAction::ActionsWithModifiedTriggers;
 
 void UInputAction::PostEditChangeProperty(FPropertyChangedEvent& PropertyChangedEvent)
 {
@@ -121,11 +122,14 @@ void UInputAction::PostEditChangeProperty(FPropertyChangedEvent& PropertyChanged
 	}
 	
 	// If our value type changes we need to inform any blueprint InputActionEx nodes that refer to this action
-	if (PropertyChangedEvent.MemberProperty->GetFName() == GET_MEMBER_NAME_CHECKED(UInputAction, ValueType) ||
-		PropertyChangedEvent.MemberProperty->GetFName() == GET_MEMBER_NAME_CHECKED(UInputAction, Triggers))
+	if (PropertyChangedEvent.MemberProperty->GetFName() == GET_MEMBER_NAME_CHECKED(UInputAction, ValueType))
 	{
 		ActionsWithModifiedValueTypes.Add(this);
 	}
+	if (PropertyChangedEvent.MemberProperty->GetFName() == GET_MEMBER_NAME_CHECKED(UInputAction, Triggers))
+	{
+		ActionsWithModifiedTriggers.Add(this);
+	}
 }
 
 EDataValidationResult UInputAction::IsDataValid(TArray<FText>& ValidationErrors)
diff --git a/Engine/Plugins/EnhancedInput/Source/EnhancedInput/Public/InputAction.h b/Engine/Plugins/EnhancedInput/Source/EnhancedInput/Public/InputAction.h
--- a/Engine/Plugins/EnhancedInput/Source/EnhancedInput/Public/InputAction.h
+++ b/Engine/Plugins/EnhancedInput/Source/EnhancedInput/Public/InputAction.h
@@ -29,6 +29,7 @@ public:
 	virtual EDataValidationResult IsDataValid(TArray<FText>& ValidationErrors) override;
 	
 	static TSet<const UInputAction*> ActionsWithModifiedValueTypes;
+	static TSet<const UInputAction*> ActionsWithModifiedTriggers;
 	
 	/**
 	 * Returns a bitmask of supported trigger events that is built from each UInputTrigger on this Action.
diff --git a/Engine/Plugins/EnhancedInput/Source/InputEditor/Private/InputEditorModule.cpp b/Engine/Plugins/EnhancedInput/Source/InputEditor/Private/InputEditorModule.cpp
--- a/Engine/Plugins/EnhancedInput/Source/InputEditor/Private/InputEditorModule.cpp
+++ b/Engine/Plugins/EnhancedInput/Source/InputEditor/Private/InputEditorModule.cpp
@@ -5,6 +5,7 @@
 #include "AssetBlueprintGraphActions.h"
 #include "AssetTypeActions_Base.h"
 #include "BlueprintEditorModule.h"
+#include "BlueprintNodeTemplateCache.h"
 #include "EnhancedInputModule.h"
 #include "DetailCategoryBuilder.h"
 #include "DetailLayoutBuilder.h"
@@ -561,34 +562,59 @@ void FInputEditorModule::AutoUpgradeDefaultInputClasses()
 void FInputEditorModule::Tick(float DeltaTime)
 {
 	// Update any blueprints that are referencing an input action with a modified value type
-	if (UInputAction::ActionsWithModifiedValueTypes.Num())
+	if (UInputAction::ActionsWithModifiedValueTypes.Num() || UInputAction::ActionsWithModifiedTriggers.Num())
 	{
-		TSet<UBlueprint*> BPsModified;
+		TSet<UBlueprint*> BPsModifiedFromValueTypeChange;
+		TSet<UBlueprint*> BPsModifiedFromTriggerChange;
+		
 		for (TObjectIterator<UK2Node_EnhancedInputAction> NodeIt; NodeIt; ++NodeIt)
 		{
-			if (UInputAction::ActionsWithModifiedValueTypes.Contains(NodeIt->InputAction))
+			if (!FBlueprintNodeTemplateCache::IsTemplateOuter(NodeIt->GetGraph()))
 			{
-				NodeIt->ReconstructNode();
-				BPsModified.Emplace(NodeIt->GetBlueprint());
+				if (UInputAction::ActionsWithModifiedValueTypes.Contains(NodeIt->InputAction))
+				{
+					NodeIt->ReconstructNode();
+					BPsModifiedFromValueTypeChange.Emplace(NodeIt->GetBlueprint());
+				}
+				if (UInputAction::ActionsWithModifiedTriggers.Contains(NodeIt->InputAction))
+				{
+					NodeIt->ReconstructNode();
+					BPsModifiedFromTriggerChange.Emplace(NodeIt->GetBlueprint());
+				}
 			}
 		}
 		for (TObjectIterator<UK2Node_GetInputActionValue> NodeIt; NodeIt; ++NodeIt)
 		{
-			if (UInputAction::ActionsWithModifiedValueTypes.Contains(NodeIt->InputAction))
+			if (!FBlueprintNodeTemplateCache::IsTemplateOuter(NodeIt->GetGraph()))
 			{
-				NodeIt->ReconstructNode();
-				BPsModified.Emplace(NodeIt->GetBlueprint());
+				if (UInputAction::ActionsWithModifiedValueTypes.Contains(NodeIt->InputAction))
+				{
+					NodeIt->ReconstructNode();
+					BPsModifiedFromValueTypeChange.Emplace(NodeIt->GetBlueprint());
+				}
+				if (UInputAction::ActionsWithModifiedTriggers.Contains(NodeIt->InputAction))
+				{
+					NodeIt->ReconstructNode();
+					BPsModifiedFromTriggerChange.Emplace(NodeIt->GetBlueprint());
+				}
 			}
 		}
 
-		if (BPsModified.Num())
+		if (BPsModifiedFromValueTypeChange.Num())
+		{
+			FNotificationInfo Info(FText::Format(LOCTEXT("ActionValueTypeChange", "Changing action value type affected {0} blueprint(s)!"), BPsModifiedFromValueTypeChange.Num()));
+			Info.ExpireDuration = 5.0f;
+			FSlateNotificationManager::Get().AddNotification(Info);
+		}
+		if (BPsModifiedFromTriggerChange.Num())
 		{
-			FNotificationInfo Info(FText::Format(LOCTEXT("ActionValueTypeChange", "Changing action value type affected {0} blueprint(s)!"), BPsModified.Num()));
+			FNotificationInfo Info(FText::Format(LOCTEXT("ActionTriggerChange", "Changing action triggers affected {0} blueprint(s)!"), BPsModifiedFromTriggerChange.Num()));
 			Info.ExpireDuration = 5.0f;
 			FSlateNotificationManager::Get().AddNotification(Info);
 		}
 
 		UInputAction::ActionsWithModifiedValueTypes.Reset();
+		UInputAction::ActionsWithModifiedTriggers.Reset();
 	}
 }
 
