Creating new Enhanced Input Combo Key and adding it at module startup

diff --git a/Engine/Plugins/EnhancedInput/Source/EnhancedInput/Private/EnhancedInputModule.cpp b/Engine/Plugins/EnhancedInput/Source/EnhancedInput/Private/EnhancedInputModule.cpp
--- a/Engine/Plugins/EnhancedInput/Source/EnhancedInput/Private/EnhancedInputModule.cpp
+++ b/Engine/Plugins/EnhancedInput/Source/EnhancedInput/Private/EnhancedInputModule.cpp
@@ -24,6 +24,8 @@
 
 #define LOCTEXT_NAMESPACE "EnhancedInput"
 
+const FKey FEnhancedInputKeys::ComboKey("ComboKey");
+
 DEFINE_LOG_CATEGORY(LogEnhancedInput);
 
 class FEnhancedInputModule : public IEnhancedInputModule, public FTickableGameObject
@@ -233,6 +235,10 @@ void FEnhancedInputModule::StartupModule()
 {
 	Library = NewObject<UEnhancedInputLibrary>(GetTransientPackage(), UEnhancedInputLibrary::StaticClass(), NAME_None);
 	Library->AddToRoot();
+	
+	const FName NAME_EnhancedInput(TEXT("EnhancedInput"));
+	EKeys::AddMenuCategoryDisplayInfo(NAME_EnhancedInput, LOCTEXT("EnhancedInputSubCateogry", "Enhanced Input"), TEXT("GraphEditor.KeyEvent_16x"));
+	EKeys::AddKey(FKeyDetails(FEnhancedInputKeys::ComboKey, LOCTEXT("ComboKey", "Combo Key"), FKeyDetails::NotActionBindableKey, NAME_EnhancedInput));
 
 	if (!IsRunningDedicatedServer())
 	{
diff --git a/Engine/Plugins/EnhancedInput/Source/EnhancedInput/Public/EnhancedInputModule.h b/Engine/Plugins/EnhancedInput/Source/EnhancedInput/Public/EnhancedInputModule.h
--- a/Engine/Plugins/EnhancedInput/Source/EnhancedInput/Public/EnhancedInputModule.h
+++ b/Engine/Plugins/EnhancedInput/Source/EnhancedInput/Public/EnhancedInputModule.h
@@ -11,6 +11,12 @@ class APlayerController;
 
 DECLARE_LOG_CATEGORY_EXTERN(LogEnhancedInput, Log, All);
 
+struct ENHANCEDINPUT_API FEnhancedInputKeys
+{
+	// Combo FKey that serves as the key combo triggers are automatically mapped to - is not action bindable
+	static const FKey ComboKey;
+};
+
 /**
  * The public interface to this module
  */
