Make the all the Enhanced Action Mapping struct properties BlueprintReadWrite instead of just BlueprintRead so that they show up on the make/break/set nodes in the BP graph. There isn't any reason why you can't set these properties via BP.

#rnx
#pf skip

diff --git a/Engine/Plugins/EnhancedInput/Source/EnhancedInput/Public/EnhancedActionKeyMapping.h b/Engine/Plugins/EnhancedInput/Source/EnhancedInput/Public/EnhancedActionKeyMapping.h
--- a/Engine/Plugins/EnhancedInput/Source/EnhancedInput/Public/EnhancedActionKeyMapping.h
+++ b/Engine/Plugins/EnhancedInput/Source/EnhancedInput/Public/EnhancedActionKeyMapping.h
@@ -60,19 +60,19 @@ public:
 	};
 
 	/** Metadata that can used to store any other related items to this key mapping such as icons, ability assets, etc. */
-	UPROPERTY(EditAnywhere, BlueprintReadOnly, Category = "Input|PlayerMappable")
+	UPROPERTY(EditAnywhere, BlueprintReadWrite, Category = "Input|PlayerMappable")
 	TObjectPtr<UObject> Metadata = nullptr;
 	
 	/** A unique name for this player mapping to be saved with. */
-	UPROPERTY(EditAnywhere, BlueprintReadOnly, Category = "Input|PlayerMappable")
+	UPROPERTY(EditAnywhere, BlueprintReadWrite, Category = "Input|PlayerMappable")
 	FName Name;
 	
 	/** The localized display name of this key mapping. Use this when displaying the mappings to a user. */
-	UPROPERTY(EditAnywhere, BlueprintReadOnly, Category = "Input|PlayerMappable")
+	UPROPERTY(EditAnywhere, BlueprintReadWrite, Category = "Input|PlayerMappable")
 	FText DisplayName;
 
 	/** The category that this player mapping is in */
-	UPROPERTY(EditAnywhere, BlueprintReadOnly, Category = "Input|PlayerMappable")
+	UPROPERTY(EditAnywhere, BlueprintReadWrite, Category = "Input|PlayerMappable")
 	FText DisplayCategory = FText::GetEmpty();
 };
 
@@ -107,7 +107,7 @@ struct FEnhancedActionKeyMapping
 #endif
 
 	/** Options for making this a player mappable keymapping */
-	UPROPERTY(EditAnywhere, BlueprintReadOnly, Category = "Input|PlayerMappable", meta = (editCondition = "bIsPlayerMappable", DisplayAfter = "bIsPlayerMappable"))
+	UPROPERTY(EditAnywhere, BlueprintReadWrite, Category = "Input|PlayerMappable", meta = (editCondition = "bIsPlayerMappable", DisplayAfter = "bIsPlayerMappable"))
 	FPlayerMappableKeyOptions PlayerMappableOptions;
 
 	/**
@@ -126,11 +126,11 @@ struct FEnhancedActionKeyMapping
 	TArray<TObjectPtr<UInputModifier>> Modifiers;
 	
 	/** Action to be affected by the key  */
-	UPROPERTY(EditAnywhere, BlueprintReadOnly, Category = "Input")
+	UPROPERTY(EditAnywhere, BlueprintReadWrite, Category = "Input")
 	TObjectPtr<const UInputAction> Action = nullptr;
 
 	/** Key that affect the action. */
-	UPROPERTY(EditAnywhere, BlueprintReadOnly, Category = "Input")
+	UPROPERTY(EditAnywhere, BlueprintReadWrite, Category = "Input")
 	FKey Key;
 
 	/**
@@ -175,7 +175,7 @@ protected:
 	TObjectPtr<UPlayerMappableKeySettings> PlayerMappableKeySettings = nullptr;
 
 	/** If true then this ActionKeyMapping will be exposed as a player mappable key */
-	UPROPERTY(EditAnywhere, BlueprintReadOnly, Category = "Input|PlayerMappable")
+	UPROPERTY(EditAnywhere, BlueprintReadWrite, Category = "Input|PlayerMappable")
 	uint8 bIsPlayerMappable : 1;
 
 public:
