Combo UX - using data validation to avoid mixing combo triggers with other triggers

diff --git a/Engine/Plugins/EnhancedInput/Source/EnhancedInput/Private/EnhancedActionKeyMapping.cpp b/Engine/Plugins/EnhancedInput/Source/EnhancedInput/Private/EnhancedActionKeyMapping.cpp
--- a/Engine/Plugins/EnhancedInput/Source/EnhancedInput/Private/EnhancedActionKeyMapping.cpp
+++ b/Engine/Plugins/EnhancedInput/Source/EnhancedInput/Private/EnhancedActionKeyMapping.cpp
@@ -65,10 +65,15 @@ EDataValidationResult FEnhancedActionKeyMapping::IsDataValid(TArray<FText>& Vali
 	}
 
 	// Validate the triggers.
+	bool bContextContainsComboTrigger = false;
+	bool bContextContainsNonComboTrigger = false;
 	for (const TObjectPtr<UInputTrigger> Trigger : Triggers)
 	{
 		if (Trigger != nullptr)
 		{
+			// check if it the trigger is a combo or not
+			Trigger.IsA(UInputTriggerCombo::StaticClass()) ? bContextContainsComboTrigger = true : bContextContainsNonComboTrigger = true;
+			
 			Result = CombineDataValidationResults(Result, Trigger->IsDataValid(ValidationErrors));
 		}
 		else
@@ -78,6 +83,58 @@ EDataValidationResult FEnhancedActionKeyMapping::IsDataValid(TArray<FText>& Vali
 		}
 	}
 
+	bool bInputActionContainsComboTrigger = false;
+    bool bInputActionContainsNonComboTrigger = false;
+	// we also need to check the input action triggers for combo triggers
+	for (const TObjectPtr<UInputTrigger> Trigger : Action->Triggers)
+	{
+		if (Trigger != nullptr)
+		{
+			Trigger.IsA(UInputTriggerCombo::StaticClass()) ? bInputActionContainsComboTrigger = true : bInputActionContainsNonComboTrigger = true;
+		}
+	}
+
+	FFormatNamedArguments Args;
+	Args.Add("InputActionName", FText::FromString(Action.GetName()));
+	Args.Add("KeyBeingMapped", Key.GetDisplayName());
+	FText DefaultComboNonComboWarning = LOCTEXT("DefaultComboNonComboWarningText", "The mapping of {InputActionName} to {KeyBeingMapped} has a Combo Trigger ({ComboTriggerLocation}) with additional non-combo triggers ({NonComboTriggerLocation}). Mixing Combo Triggers with other types of Triggers is not supported. Consider putting the Combo Trigger(s) on a seperate mapping or making a seperate Input Action for them.");
+	if (bContextContainsComboTrigger)
+	{
+		Args.Add("ComboTriggerLocation", LOCTEXT("ComboInContextText", "From the Mapping Context"));
+		if (bContextContainsNonComboTrigger)
+		{
+			Result = EDataValidationResult::Invalid;
+			Args.Add("NonComboTriggerLocation", LOCTEXT("NonComboInContextText", "From the Mapping Context"));
+			ValidationErrors.Add(FText::Format(DefaultComboNonComboWarning, Args));
+		}
+		// Input Action contains non-combo trigger(s) 
+		if (bInputActionContainsNonComboTrigger)
+		{
+			Args.Add("NonComboTriggerLocation", LOCTEXT("NonComboInInputActionText", "From the Input Action"));
+			Result = EDataValidationResult::Invalid;
+			ValidationErrors.Add(FText::Format(DefaultComboNonComboWarning, Args));
+		}
+	}
+	// Input Action contains combo trigger(s)
+	if (bInputActionContainsComboTrigger)
+	{
+		Args.Add("ComboTriggerLocation", LOCTEXT("ComboInInputActionText", "From the Input Action"));
+		// Context contains non-combo trigger(s)
+		if (bContextContainsNonComboTrigger)
+		{
+			Result = EDataValidationResult::Invalid;
+			Args.Add("NonComboTriggerLocation", LOCTEXT("NonComboInContextText", "From the Mapping Context"));
+			ValidationErrors.Add(FText::Format(DefaultComboNonComboWarning, Args));
+		}
+		// Input Action contains non-combo trigger(s) 
+		if (bInputActionContainsNonComboTrigger)
+		{
+			Args.Add("NonComboTriggerLocation", LOCTEXT("NonComboInInputActionText", "From the Input Action"));
+			Result = EDataValidationResult::Invalid;
+			ValidationErrors.Add(FText::Format(DefaultComboNonComboWarning, Args));
+		}
+	}
+	
 	// Validate the modifiers.
 	for (const TObjectPtr<UInputModifier> Modifier : Modifiers)
 	{
diff --git a/Engine/Plugins/EnhancedInput/Source/EnhancedInput/Private/InputAction.cpp b/Engine/Plugins/EnhancedInput/Source/EnhancedInput/Private/InputAction.cpp
--- a/Engine/Plugins/EnhancedInput/Source/EnhancedInput/Private/InputAction.cpp
+++ b/Engine/Plugins/EnhancedInput/Source/EnhancedInput/Private/InputAction.cpp
@@ -133,10 +133,15 @@ EDataValidationResult UInputAction::IsDataValid(TArray<FText>& ValidationErrors)
 	EDataValidationResult Result = CombineDataValidationResults(Super::IsDataValid(ValidationErrors), EDataValidationResult::Valid);
 
 	// Validate the triggers
+	bool bContainsComboTrigger = false;
+	bool bContainsNonComboTrigger = false;
 	for (const TObjectPtr<UInputTrigger> Trigger : Triggers)
 	{
 		if (Trigger)
 		{
+			// check if it the trigger is a combo or not
+			Trigger.IsA(UInputTriggerCombo::StaticClass()) ? bContainsComboTrigger = true : bContainsNonComboTrigger = true;
+			
 			Result = CombineDataValidationResults(Result, Trigger->IsDataValid(ValidationErrors));
 		}
 		else
@@ -145,6 +150,11 @@ EDataValidationResult UInputAction::IsDataValid(TArray<FText>& ValidationErrors)
 			ValidationErrors.Add(LOCTEXT("NullInputTrigger", "There cannot be a null Input Trigger on an Input Action!"));
 		}
 	}
+	if (bContainsComboTrigger && bContainsNonComboTrigger)
+	{
+		Result = EDataValidationResult::Invalid;
+		ValidationErrors.Add(LOCTEXT("ComboAndNonComboInputTrigger", "Combo triggers are not intended to interact with other input triggers. Consider adding the combo and other triggers later in a context or creating a seperate action for the combo."));
+	}
 	
 	// Validate the modifiers
 	for (const TObjectPtr<UInputModifier> Modifier : Modifiers)
