Combo trigger Improvements.
Fixing issues with combo triggers being overriden
Can specify events that progress/cancel combo
some smaller UX fixes

diff --git a/Engine/Plugins/EnhancedInput/Source/EnhancedInput/Private/InputTriggers.cpp b/Engine/Plugins/EnhancedInput/Source/EnhancedInput/Private/InputTriggers.cpp
--- a/Engine/Plugins/EnhancedInput/Source/EnhancedInput/Private/InputTriggers.cpp
+++ b/Engine/Plugins/EnhancedInput/Source/EnhancedInput/Private/InputTriggers.cpp
@@ -241,12 +241,14 @@ ETriggerState UInputTriggerCombo::UpdateState_Implementation(const UEnhancedPlay
 	if (const UInputAction* CurrentAction = ComboActions[CurrentComboStepIndex].ComboStepAction)
 	{
 		// loop through all cancel actions and check if they've fired
-		for (const UInputAction* CancelAction : CancelActions)
+		for (FInputCancelAction InputCancelActionData : InputCancelActions)
 		{
+			const UInputAction* CancelAction = InputCancelActionData.CancelAction;
 			if (CancelAction && CancelAction != CurrentAction)
 			{
 				const FInputActionInstance* CancelState = PlayerInput->FindActionInstanceData(CancelAction);
-				if (CancelState && CancelState->GetTriggerEvent() != ETriggerEvent::None)
+				// Check the cancellation state against the states that should cancel the combo
+				if (CancelState && (InputCancelActionData.CancellationStates & static_cast<uint8>(CancelState->GetTriggerEvent())))
 				{
 					// Cancel action firing!
 					CurrentComboStepIndex = 0;
@@ -261,7 +263,8 @@ ETriggerState UInputTriggerCombo::UpdateState_Implementation(const UEnhancedPlay
 			if (ComboStep.ComboStepAction && ComboStep.ComboStepAction != CurrentAction)
 			{
 				const FInputActionInstance* CancelState = PlayerInput->FindActionInstanceData(ComboStep.ComboStepAction);
-				if (CancelState && CancelState->GetTriggerEvent() != ETriggerEvent::None)
+				// Check the combo action state against the states that should complete this step
+				if (CancelState && (ComboStep.ComboStepCompletionStates & static_cast<uint8>(CancelState->GetTriggerEvent())))
 				{
 					// Other combo action firing - should cancel
 					CurrentComboStepIndex = 0;
@@ -283,8 +286,8 @@ ETriggerState UInputTriggerCombo::UpdateState_Implementation(const UEnhancedPlay
 		}
 
 		const FInputActionInstance* CurrentState = PlayerInput->FindActionInstanceData(CurrentAction);
-		// check to see if current action is completed - if so advance the combo to the next combo action
-		if (CurrentState && CurrentState->GetTriggerEvent() == ETriggerEvent::Completed) // + possibly Triggered
+		// check to see if current action is in one of it's completion states - if so advance the combo to the next combo action
+		if (CurrentState && (ComboActions[CurrentComboStepIndex].ComboStepCompletionStates & static_cast<uint8>(CurrentState->GetTriggerEvent())))
 		{
 			CurrentComboStepIndex++;
 			CurrentTimeBetweenComboSteps = 0;
@@ -327,6 +330,30 @@ EDataValidationResult UInputTriggerCombo::IsDataValid(TArray<FText>& ValidationE
 	return Result;
 }
 
+void UInputTriggerCombo::PostLoad()
+{
+	Super::PostLoad();
+	
+	// start fix up in case it's an old version that just has the cancel input action
+#if WITH_EDITORONLY_DATA
+	PRAGMA_DISABLE_DEPRECATION_WARNINGS
+
+	if (!CancelActions.IsEmpty())
+	{
+		for (TObjectPtr<const UInputAction> InputAction : CancelActions)
+		{
+			// use default settings but set cancel action
+			FInputCancelAction InputCancelAction;
+			InputCancelAction.CancelAction = InputAction;
+			InputCancelActions.Add(InputCancelAction);
+		}
+		CancelActions.Empty();
+	}
+		
+	PRAGMA_ENABLE_DEPRECATION_WARNINGS
+#endif // WITH_EDITORONLY_DATA
+}
+
 #endif
 
 #undef LOCTEXT_NAMESPACE
diff --git a/Engine/Plugins/EnhancedInput/Source/EnhancedInput/Public/InputTriggers.h b/Engine/Plugins/EnhancedInput/Source/EnhancedInput/Public/InputTriggers.h
--- a/Engine/Plugins/EnhancedInput/Source/EnhancedInput/Public/InputTriggers.h
+++ b/Engine/Plugins/EnhancedInput/Source/EnhancedInput/Public/InputTriggers.h
@@ -28,29 +28,30 @@ enum class ETriggerState : uint8
 /**
 * Trigger events are the Action's interpretation of all Trigger State transitions that occurred for the action in the last tick
 */
-UENUM(BlueprintType)
+UENUM(BlueprintType, meta = (Bitflags, UseEnumValuesAsMaskValuesInEditor = "true"))
 enum class ETriggerEvent : uint8
 {
 	// No significant trigger state changes occurred and there are no active device inputs
-	None = 0				UMETA(Hidden),
+	None		= (0x0)		UMETA(Hidden),
 
 	// Triggering occurred after one or more processing ticks
-	Triggered,				// ETriggerState (None -> Triggered, Ongoing -> Triggered, Triggered -> Triggered)
+	Triggered	= (1 << 0),	// ETriggerState (None -> Triggered, Ongoing -> Triggered, Triggered -> Triggered)
 	
 	// An event has occurred that has begun Trigger evaluation. Note: Triggered may also occur this frame.
-	Started,				// ETriggerState (None -> Ongoing, None -> Triggered)
+	Started		= (1 << 1),	// ETriggerState (None -> Ongoing, None -> Triggered)
 
 	// Triggering is still being processed
-	Ongoing,				// ETriggerState (Ongoing -> Ongoing)
+	Ongoing		= (1 << 2),	// ETriggerState (Ongoing -> Ongoing)
 
 	// Triggering has been canceled
-	Canceled,				// ETriggerState (Ongoing -> None)
+	Canceled	= (1 << 3),	// ETriggerState (Ongoing -> None)
 
 	// The trigger state has transitioned from Triggered to None this frame, i.e. Triggering has finished.
 	// NOTE: Using this event restricts you to one set of triggers for Started/Completed events. You may prefer two actions, each with its own trigger rules.
 	// TODO: Completed will not fire if any trigger reports Ongoing on the same frame, but both should fire. e.g. Tick 2 of Hold (= Ongoing) + Pressed (= None) combo will raise Ongoing event only.
-	Completed,				// ETriggerState (Triggered -> None)
+	Completed	= (1 << 4),	// ETriggerState (Triggered -> None)
 };
+ENUM_CLASS_FLAGS(ETriggerEvent)
 
 /**
 * Trigger type determine how the trigger contributes to an action's overall trigger event the behavior of the trigger
@@ -417,20 +418,43 @@ struct FInputComboStepData
 {
 	GENERATED_BODY()
 
-	// The action that must be triggering to progress the combo
+	// The action that must be completed (according to Combo Step Completion States) to progress the combo
 	UPROPERTY(EditInstanceOnly, BlueprintReadWrite, Category = "Trigger Settings", meta = (DisplayThumbnail = "false"))
 	TObjectPtr<const UInputAction> ComboStepAction = nullptr;
 
-	// Time to press the key and for it to count towards the combo
+	// Trigger events that will complete this step - what events from this action should progress the combo
+	UPROPERTY(EditAnywhere, BlueprintReadWrite, Category = "Trigger Settings", config, meta = (Bitmask, BitmaskEnum="/Script/EnhancedInput.ETriggerEvent"))
+	uint8 ComboStepCompletionStates = static_cast<uint8>(ETriggerEvent::Triggered);
+	
+	/**
+	 * Time to press the key before combo is cancelled - starts once the previous step in the combo is completed
+	 * Note: This can be safely ignored for the first action in the combo
+	 */
 	UPROPERTY(EditAnywhere, Config, BlueprintReadWrite, Category = "Trigger Settings")
 	float TimeToPressKey = 0.5f;
 };
 
+USTRUCT(BlueprintType)
+struct FInputCancelAction
+{
+	GENERATED_BODY()
+
+	// The action that must be completed (according to Cancellation States) to cancel the combo
+	UPROPERTY(EditInstanceOnly, BlueprintReadWrite, Category = "Trigger Settings", meta = (DisplayThumbnail = "false"))
+	TObjectPtr<const UInputAction> CancelAction = nullptr;
+
+	// Trigger events for this action that will cancel the combo - what events from this action should cancel the combo
+	UPROPERTY(EditAnywhere, BlueprintReadWrite, Category = "Trigger Settings", config, meta = (Bitmask, BitmaskEnum="/Script/EnhancedInput.ETriggerEvent"))
+	uint8 CancellationStates = static_cast<uint8>(ETriggerEvent::Triggered);
+};
+
 /**
  * UInputTriggerCombo
- * All actions in the combo array must be pressed within a timeframe to trigger
-*/
-UCLASS(NotBlueprintable, meta = (DisplayName = "Combo", NotInputConfigurable = "true"))
+ * All actions in the combo array must be completed (based on combo completion event specified - triggered, completed, etc.) to trigger the action this trigger is on.
+ * Actions must also be completed in the order specified by the combo actions array (starting at index 0).
+ * Note: This will only trigger for one frame before resetting the combo trigger's progress 
+ */
+UCLASS(NotBlueprintable, meta = (DisplayName = "Combo (Beta)", NotInputConfigurable = "true"))
 class ENHANCEDINPUT_API UInputTriggerCombo : public UInputTrigger
 {
 	GENERATED_BODY()
@@ -456,17 +480,25 @@ public:
 #if WITH_EDITOR
 	virtual EDataValidationResult IsDataValid(TArray<FText>& ValidationErrors) override;
 #endif
+
+#if  WITH_EDITORONLY_DATA
+	// Actions that will cancel the combo if they are triggered.
+	UE_DEPRECATED(5.2, "CancelActions has been deprecated as of 5.2. Please use InputCancelActions instead.")
+	TArray<TObjectPtr<const UInputAction>> CancelActions;
+#endif
+
+	virtual void PostLoad() override;
 	
 	/**
-	 * List of input actions that need to be completed to trigger this action.
-	 * Input actions must be triggered in order (starting at index 0) to count towards the completion of the combo.
+	 * List of input actions that need to be completed (according to Combo Step Completion States) to trigger this action.
+	 * Input actions must be triggered in order (starting at index 0) to count towards the triggering of the combo.
 	 */
 	UPROPERTY(EditAnywhere, BlueprintReadWrite, Category = "Trigger Settings", meta = (DisplayThumbnail = "false", TitleProperty = "ComboStepAction"))
 	TArray<FInputComboStepData> ComboActions;
 
-	// Actions that will cancel the combo if they are triggered
-    UPROPERTY(EditAnywhere, BlueprintReadWrite, Category = "Trigger Settings", meta = (DisplayThumbnail = "false"))
-    TArray<TObjectPtr<const UInputAction>> CancelActions;
+	// Actions that will cancel the combo if they are completed (according to Cancellation States)
+	UPROPERTY(EditAnywhere, BlueprintReadWrite, Category = "Trigger Settings", meta = (DisplayThumbnail = "false", DisplayName = "Cancel Actions"))
+	TArray<FInputCancelAction> InputCancelActions;
 
 	/** Determines what kind of trigger events can happen from the behavior of this trigger. */
 	virtual ETriggerEventsSupported GetSupportedTriggerEvents() const override { return ETriggerEventsSupported::All; }
diff --git a/Engine/Plugins/EnhancedInput/Source/InputBlueprintNodes/Private/K2Node_EnhancedInputAction.cpp b/Engine/Plugins/EnhancedInput/Source/InputBlueprintNodes/Private/K2Node_EnhancedInputAction.cpp
--- a/Engine/Plugins/EnhancedInput/Source/InputBlueprintNodes/Private/K2Node_EnhancedInputAction.cpp
+++ b/Engine/Plugins/EnhancedInput/Source/InputBlueprintNodes/Private/K2Node_EnhancedInputAction.cpp
@@ -68,7 +68,7 @@ void UK2Node_EnhancedInputAction::AllocateDefaultPins()
 		static const UEnum* EventEnum = StaticEnum<ETriggerEvent>();
 
 		UEdGraphPin* NewPin = CreatePin(EGPD_Output, UEdGraphSchema_K2::PC_Exec, PinName);
-		NewPin->PinToolTip = EventEnum->GetToolTipTextByIndex((int32)Event).ToString();
+		NewPin->PinToolTip = EventEnum->GetToolTipTextByIndex(EventEnum->GetIndexByValue(static_cast<uint8>(Event))).ToString();
 
 		// Add a special tooltip and display name for pins that are unsupported
 		if (UE::Input::bShouldWarnOnUnsupportedInputPin && !UInputTrigger::IsSupportedTriggerEvent(SupportedTriggerEvents, Event))
diff --git a/Engine/Plugins/EnhancedInput/Source/InputEditor/Private/InputCustomizations.cpp b/Engine/Plugins/EnhancedInput/Source/InputEditor/Private/InputCustomizations.cpp
--- a/Engine/Plugins/EnhancedInput/Source/InputEditor/Private/InputCustomizations.cpp
+++ b/Engine/Plugins/EnhancedInput/Source/InputEditor/Private/InputCustomizations.cpp
@@ -65,7 +65,7 @@ void FEnhancedActionMappingCustomization::CustomizeHeader(TSharedRef<IPropertyHa
 	TSharedPtr<IPropertyHandle> TriggersHandle = PropertyHandle->GetChildHandle(GET_MEMBER_NAME_CHECKED(FEnhancedActionKeyMapping, Triggers));
 
 	TriggersHandle->SetOnPropertyValueChanged(FSimpleDelegate::CreateSP(this, &FEnhancedActionMappingCustomization::OnTriggersChanged));
-
+	TriggersHandle->SetOnChildPropertyValueChanged(FSimpleDelegate::CreateSP(this, &FEnhancedActionMappingCustomization::OnTriggersChanged));
 	TSharedRef<SWidget> RemoveButton = PropertyCustomizationHelpers::MakeDeleteButton(FSimpleDelegate::CreateSP(this, &FEnhancedActionMappingCustomization::RemoveMappingButton_OnClick),
 		LOCTEXT("RemoveMappingToolTip", "Remove Mapping"));
 
@@ -74,10 +74,11 @@ void FEnhancedActionMappingCustomization::CustomizeHeader(TSharedRef<IPropertyHa
 
 	// TODO: Use FDetailArrayBuilder?
 
+	const bool bContainsComboTrigger = DoesTriggerArrayContainCombo();
 	// Pass our header row into the key struct customizeheader method so it populates our row with the key struct header
 	KeyStructCustomization = StaticCastSharedPtr<FKeyStructCustomization>(KeyStructInstance);
-	const bool bContainsComboTrigger = DoesTriggerArrayContainCombo();
-	KeyStructCustomization->SetDisplayIcon(bContainsComboTrigger);
+	KeyStructCustomization->SetDefaultKeyName("ComboKey");
+	KeyStructCustomization->SetDisabledKeySelectorToolTip(LOCTEXT("DisabledSelectorComboToolTip", "A Combo Trigger isn't triggered through a key so Key Selection has been disabled."));
 	KeyStructCustomization->SetEnableKeySelector(!bContainsComboTrigger);
 	KeyStructCustomization->CustomizeHeaderOnlyWithButton(KeyHandle.ToSharedRef(), HeaderRow, CustomizationUtils, RemoveButton);
 }
@@ -167,8 +168,17 @@ void FEnhancedActionMappingCustomization::RemoveMappingButton_OnClick() const
 void FEnhancedActionMappingCustomization::OnTriggersChanged() const
 {
 	const bool bContainsComboTrigger = DoesTriggerArrayContainCombo();
-	KeyStructCustomization->SetDisplayIcon(bContainsComboTrigger);
+	
+	// if it is currently disabled and doesn't contain a combo now we should set the key to something other than ComboKey
+	if (!KeyStructCustomization->GetEnableKeySelector() && !bContainsComboTrigger)
+	{
+		// setting KeySelector to none key
+		KeyStructCustomization->SetKey(TEXT("None"));
+	}
+	
     KeyStructCustomization->SetEnableKeySelector(!bContainsComboTrigger);
+	// updating the default key when the KeySelector is disabled
+	KeyStructCustomization->SetDefaultKeyName(bContainsComboTrigger ? TEXT("ComboKey") : TEXT("None"));
 }
 
 void FEnhancedActionMappingCustomization::OnInputActionTriggersChanged() const
@@ -213,29 +223,24 @@ void FEnhancedActionMappingCustomization::OnInputActionModifiersChanged() const
 
 bool FEnhancedActionMappingCustomization::DoesTriggerArrayContainCombo() const
 {
-	if (MappingPropertyHandle)
+	if (const FEnhancedActionKeyMapping* ActionKeyMapping = GetActionKeyMapping(MappingPropertyHandle))
 	{
-		if (TSharedPtr<IPropertyHandle> TriggersHandle = MappingPropertyHandle->GetChildHandle(GET_MEMBER_NAME_CHECKED(FEnhancedActionKeyMapping, Triggers)))
+		// checking context triggers for combo triggers
+		for (TObjectPtr<UInputTrigger> Trigger : ActionKeyMapping->Triggers)
 		{
-			// getting data for the trigger array
-			void* Data = nullptr;
-			TriggersHandle->GetValueData(Data);
-			if (Data)
+			if (Trigger.IsA(UInputTriggerCombo::StaticClass()))
 			{
-				FProperty* TriggersProperty = TriggersHandle->GetProperty();
-				FArrayProperty* TriggersArrayProperty = CastField<FArrayProperty>(TriggersProperty);
-				FScriptArrayHelper ArrayHelper(TriggersArrayProperty, Data);
-	
-				for (int32 i = 0; i < ArrayHelper.Num(); i++)
+				return true;
+			}
+		}
+		// checking input action triggers for combo triggers
+		if (ActionKeyMapping->Action)
+		{
+			for (TObjectPtr<UInputTrigger> Trigger : ActionKeyMapping->Action->Triggers)
+			{
+				if (Trigger.IsA(UInputTriggerCombo::StaticClass()))
 				{
-					// Make sure we can cast this to a input trigger and if it's a combo we can return true
-					if (UInputTrigger** InputComboTrigger = reinterpret_cast<UInputTrigger**>(ArrayHelper.GetRawPtr(i)))
-					{
-						if ((*InputComboTrigger) && (*InputComboTrigger)->IsA(UInputTriggerCombo::StaticClass()))
-						{
-							return true;
-						}
-					}
+					return true;
 				}
 			}
 		}
diff --git a/Engine/Source/Editor/DetailCustomizations/Private/KeyStructCustomization.cpp b/Engine/Source/Editor/DetailCustomizations/Private/KeyStructCustomization.cpp
--- a/Engine/Source/Editor/DetailCustomizations/Private/KeyStructCustomization.cpp
+++ b/Engine/Source/Editor/DetailCustomizations/Private/KeyStructCustomization.cpp
@@ -35,6 +35,11 @@ void FKeyStructCustomization::SetEnableKeySelector(bool bKeySelectorEnabled)
 	}
 }
 
+void FKeyStructCustomization::SetKey(const FString& KeyName)
+{
+	PropertyHandle->SetValueFromFormattedString(KeyName);
+}
+
 TSharedRef<IPropertyTypeCustomization> FKeyStructCustomization::MakeInstance( )
 {
 	return MakeShareable(new FKeyStructCustomization);
@@ -81,6 +86,7 @@ void FKeyStructCustomization::CustomizeHeaderOnlyWithButton(TSharedRef<class IPr
 	    });
 	
 	KeySelector->SetEnabledFromKeyStructCustomization(bEnableKeySelector);
+	KeySelector->SetDisabledKeySelectorToolTip(DisabledKeySelectorToolTip);
 	
 	// create struct header
 	HeaderRow.NameContent()
@@ -102,20 +108,6 @@ void FKeyStructCustomization::CustomizeHeaderOnlyWithButton(TSharedRef<class IPr
 		[
 			Button
 		]
-		+ SHorizontalBox::Slot()
-		.Padding(InputSettingsDetails::InputConstants::PropertyPadding)
-		.HAlign(HAlign_Center)
-		.VAlign(VAlign_Center)
-		.AutoWidth()
-		[
-			SNew(SImage)
-			.Visibility_Lambda([this]() {
-				return bDisplayIcon ? EVisibility::Visible : EVisibility::Hidden;
-			})
-			.Image(FAppStyle::GetBrush("Icons.Info"))
-			.ToolTipText(LOCTEXT("ComboTriggerKeyWarningText", "This Key is NOT part of the Combo. This action will be triggered when the actions in the Combo Actions array are completed."))
-			.ColorAndOpacity(FSlateColor::UseForeground())
-        ]
 	];
 }
 
@@ -123,7 +115,7 @@ TOptional<FKey> FKeyStructCustomization::GetCurrentKey() const
 {
 	if (!bEnableKeySelector)
 	{
-		PropertyHandle->SetValueFromFormattedString(TEXT("None"));
+		PropertyHandle->SetValueFromFormattedString(DefaultKeyName);
 	}
 	
 	TArray<void*> StructPtrs;
diff --git a/Engine/Source/Editor/DetailCustomizations/Public/KeyStructCustomization.h b/Engine/Source/Editor/DetailCustomizations/Public/KeyStructCustomization.h
--- a/Engine/Source/Editor/DetailCustomizations/Public/KeyStructCustomization.h
+++ b/Engine/Source/Editor/DetailCustomizations/Public/KeyStructCustomization.h
@@ -32,26 +32,58 @@ public:
 	// TODO: Is there a better way?
 	void CustomizeHeaderOnlyWithButton(TSharedRef<class IPropertyHandle> StructPropertyHandle, class FDetailWidgetRow& HeaderRow, IPropertyTypeCustomizationUtils& StructCustomizationUtils, TSharedRef<SWidget> Button);
 
-	/** Sets a bool for whether the key selector should be enabled or not */
-	void SetEnableKeySelector(bool bKeySelectorEnabled);
+	PRAGMA_DISABLE_DEPRECATION_WARNINGS
 	/** Sets a bool for whether the information Icon for a combo trigger is displayed or not */ 
+	UE_DEPRECATED(5.2, "SetDisplayIcon(bShouldDisplayIcon) is deprecated and the icon referenced has been removed")
 	void SetDisplayIcon(bool bShouldDisplayIcon)
 	{
 		bDisplayIcon = bShouldDisplayIcon;
 	}
 
+	/** Gets bDisplayIcon bool */
+	UE_DEPRECATED(5.2, "GetDisplayIcon() is deprecated and the icon referenced has been removed")
+	bool GetDisplayIcon() const
+	{
+		return bDisplayIcon;
+	}
+	PRAGMA_ENABLE_DEPRECATION_WARNINGS
+	
+	/** Sets a bool for whether the key selector should be enabled or not */
+	void SetEnableKeySelector(bool bKeySelectorEnabled);
+
+	
 	/** Gets bEnableKeySelector bool */
 	bool GetEnableKeySelector() const
 	{
 		return bEnableKeySelector;
 	}
 	
-	/** Gets bDisplayIcon bool */
-    bool GetDisplayIcon() const
-    {
-    	return bDisplayIcon;
-    }
+	/** Sets default key name - what key selector will default to if it is disabled */
+	void SetDefaultKeyName(FString KeyName)
+	{
+		DefaultKeyName = KeyName;
+	}
+	
+	/** Gets default key name - what key selector will default to if it is disabled */
+	FString GetDefaultKeyName() const
+	{
+		return DefaultKeyName;
+	}
 
+	/** Sets tooltip on the KeySelector when it is disabled */
+	void SetDisabledKeySelectorToolTip(const FText& InToolTip)
+	{
+		DisabledKeySelectorToolTip = InToolTip;
+	}
+	
+	/** Gets tooltip on the KeySelector when it is disabled */
+	FText GetDisabledKeySelectorToolTip() const
+	{
+		return DisabledKeySelectorToolTip;
+	}
+
+	void SetKey(const FString& KeyName);
+	
 public:
 
 	/**
@@ -63,10 +95,18 @@ public:
 
 protected:
 	/** Whether the information icon for combo triggers is displayed or not */
+	UE_DEPRECATED(5.2, "bDisplayIcon is deprecated and the icon has been removed")
 	bool bDisplayIcon = false;
+	
 	/** Whether the key selector should be enabled or not */
 	bool bEnableKeySelector = true;
 
+	/** Tooltip to display on the actual KeySelector */
+	FText DisabledKeySelectorToolTip = FText::FromString(TEXT("Key Selector Disabled"));
+	
+	/** Default Key to assign the KeySelector when it is disabled */
+	FString DefaultKeyName = TEXT("None");
+	
 	/** Gets the current Key being edited. */
 	TOptional<FKey> GetCurrentKey() const;
 
diff --git a/Engine/Source/Editor/UnrealEd/Private/SKeySelector.cpp b/Engine/Source/Editor/UnrealEd/Private/SKeySelector.cpp
--- a/Engine/Source/Editor/UnrealEd/Private/SKeySelector.cpp
+++ b/Engine/Source/Editor/UnrealEd/Private/SKeySelector.cpp
@@ -207,11 +207,6 @@ FText SKeySelector::GetKeyDescription() const
 
 FText SKeySelector::GetKeyDescriptionToolTip() const
 {
-	if(!bEnabledFromKeyStructCustomization)
-    {
-    	return LOCTEXT("KeySelectorDisabledToolTipText", "This Key cannot be part of the combo trigger so Key selection has been disabled.");
-    }
-    
 	TOptional<FKey> CurrentKeyValue = CurrentKey.Get();
 	if (CurrentKeyValue.IsSet())
 	{
@@ -239,7 +234,7 @@ FText SKeySelector::GetKeyTooltip() const
 {
 	if(!bEnabledFromKeyStructCustomization)
 	{
-		return LOCTEXT("KeySelectorDisabledToolTipText", "This Key cannot be part of the combo trigger so Key selection has been disabled.");
+		return DisabledSelectorToolTip;
 	}
 	
 	TOptional<FKey> CurrentKeyValue = CurrentKey.Get();
diff --git a/Engine/Source/Editor/UnrealEd/Public/SKeySelector.h b/Engine/Source/Editor/UnrealEd/Public/SKeySelector.h
--- a/Engine/Source/Editor/UnrealEd/Public/SKeySelector.h
+++ b/Engine/Source/Editor/UnrealEd/Public/SKeySelector.h
@@ -87,6 +87,18 @@ public:
 		return bEnabledFromKeyStructCustomization;
 	}
 	
+	/** Sets tooltip on the KeySelector when it is disabled */
+	void SetDisabledKeySelectorToolTip(const FText& InToolTip)
+	{
+		DisabledSelectorToolTip = InToolTip;
+	}
+	
+	/** Gets tooltip on the KeySelector when it is disabled */
+	FText GetDisabledKeySelectorToolTip() const
+	{
+		return DisabledSelectorToolTip;
+	}
+	
 protected:
 	/** Gets the icon for the key being manipulated */
 	const FSlateBrush* GetKeyIconImage() const;
@@ -99,6 +111,9 @@ protected:
 	FText GetKeyDescriptionToolTip() const;
 	/** Gets a tooltip for the selected key */
 	FText GetKeyTooltip() const;
+	
+	/** Tooltip to display on the selector when the selector is disabled*/
+	FText DisabledSelectorToolTip = FText::FromString(TEXT("Key Selector Disabled"));
 
 	/** Treeview support functions */
 	virtual TSharedRef<ITableRow> GenerateKeyTreeRow(FKeyTreeItem InItem, const TSharedRef<STableViewBase>& OwnerTree);
