Fixing ifeditor error with postloading a combo trigger

#rnx

diff --git a/Engine/Plugins/EnhancedInput/Source/EnhancedInput/Private/InputTriggers.cpp b/Engine/Plugins/EnhancedInput/Source/EnhancedInput/Private/InputTriggers.cpp
--- a/Engine/Plugins/EnhancedInput/Source/EnhancedInput/Private/InputTriggers.cpp
+++ b/Engine/Plugins/EnhancedInput/Source/EnhancedInput/Private/InputTriggers.cpp
@@ -2,7 +2,6 @@
 
 #include "InputTriggers.h"
 #include "EnhancedPlayerInput.h"
-#include "UObject/Object.h"
 
 #include UE_INLINE_GENERATED_CPP_BY_NAME(InputTriggers)
 
@@ -330,6 +329,7 @@ EDataValidationResult UInputTriggerCombo::IsDataValid(TArray<FText>& ValidationE
 
 	return Result;
 }
+#endif // WITH_EDITOR
 
 void UInputTriggerCombo::PostLoad()
 {
@@ -355,6 +355,4 @@ void UInputTriggerCombo::PostLoad()
 #endif // WITH_EDITORONLY_DATA
 }
 
-#endif
-
 #undef LOCTEXT_NAMESPACE
