Fix a crash when saving an IMC with a null input action. This was introduced by 23669649

#pf skip
#rnx

diff --git a/Engine/Plugins/EnhancedInput/Source/EnhancedInput/Private/EnhancedActionKeyMapping.cpp b/Engine/Plugins/EnhancedInput/Source/EnhancedInput/Private/EnhancedActionKeyMapping.cpp
--- a/Engine/Plugins/EnhancedInput/Source/EnhancedInput/Private/EnhancedActionKeyMapping.cpp
+++ b/Engine/Plugins/EnhancedInput/Source/EnhancedInput/Private/EnhancedActionKeyMapping.cpp
@@ -83,55 +83,58 @@ EDataValidationResult FEnhancedActionKeyMapping::IsDataValid(TArray<FText>& Vali
 		}
 	}
 
-	bool bInputActionContainsComboTrigger = false;
-    bool bInputActionContainsNonComboTrigger = false;
-	// we also need to check the input action triggers for combo triggers
-	for (const TObjectPtr<UInputTrigger> Trigger : Action->Triggers)
+	if (Action)
 	{
-		if (Trigger != nullptr)
+		bool bInputActionContainsComboTrigger = false;
+	    bool bInputActionContainsNonComboTrigger = false;
+		// we also need to check the input action triggers for combo triggers
+		for (const TObjectPtr<UInputTrigger> Trigger : Action->Triggers)
 		{
-			Trigger.IsA(UInputTriggerCombo::StaticClass()) ? bInputActionContainsComboTrigger = true : bInputActionContainsNonComboTrigger = true;
+			if (Trigger != nullptr)
+			{
+				Trigger.IsA(UInputTriggerCombo::StaticClass()) ? bInputActionContainsComboTrigger = true : bInputActionContainsNonComboTrigger = true;
+			}
 		}
-	}
 
-	FFormatNamedArguments Args;
-	Args.Add("InputActionName", FText::FromString(Action.GetName()));
-	Args.Add("KeyBeingMapped", Key.GetDisplayName());
-	FText DefaultComboNonComboWarning = LOCTEXT("DefaultComboNonComboWarningText", "The mapping of {InputActionName} to {KeyBeingMapped} has a Combo Trigger ({ComboTriggerLocation}) with additional non-combo triggers ({NonComboTriggerLocation}). Mixing Combo Triggers with other types of Triggers is not supported. Consider putting the Combo Trigger(s) on a seperate mapping or making a seperate Input Action for them.");
-	if (bContextContainsComboTrigger)
-	{
-		Args.Add("ComboTriggerLocation", LOCTEXT("ComboInContextText", "From the Mapping Context"));
-		if (bContextContainsNonComboTrigger)
+		FFormatNamedArguments Args;
+		Args.Add("InputActionName", FText::FromName(Action->GetFName()));
+		Args.Add("KeyBeingMapped", Key.GetDisplayName());
+		FText DefaultComboNonComboWarning = LOCTEXT("DefaultComboNonComboWarningText", "The mapping of {InputActionName} to {KeyBeingMapped} has a Combo Trigger ({ComboTriggerLocation}) with additional non-combo triggers ({NonComboTriggerLocation}). Mixing Combo Triggers with other types of Triggers is not supported. Consider putting the Combo Trigger(s) on a seperate mapping or making a seperate Input Action for them.");
+		if (bContextContainsComboTrigger)
 		{
-			Result = EDataValidationResult::Invalid;
-			Args.Add("NonComboTriggerLocation", LOCTEXT("NonComboInContextText", "From the Mapping Context"));
-			ValidationErrors.Add(FText::Format(DefaultComboNonComboWarning, Args));
+			Args.Add("ComboTriggerLocation", LOCTEXT("ComboInContextText", "From the Mapping Context"));
+			if (bContextContainsNonComboTrigger)
+			{
+				Result = EDataValidationResult::Invalid;
+				Args.Add("NonComboTriggerLocation", LOCTEXT("NonComboInContextText", "From the Mapping Context"));
+				ValidationErrors.Add(FText::Format(DefaultComboNonComboWarning, Args));
+			}
+			// Input Action contains non-combo trigger(s) 
+			if (bInputActionContainsNonComboTrigger)
+			{
+				Args.Add("NonComboTriggerLocation", LOCTEXT("NonComboInInputActionText", "From the Input Action"));
+				Result = EDataValidationResult::Invalid;
+				ValidationErrors.Add(FText::Format(DefaultComboNonComboWarning, Args));
+			}
 		}
-		// Input Action contains non-combo trigger(s) 
-		if (bInputActionContainsNonComboTrigger)
+		// Input Action contains combo trigger(s)
+		if (bInputActionContainsComboTrigger)
 		{
-			Args.Add("NonComboTriggerLocation", LOCTEXT("NonComboInInputActionText", "From the Input Action"));
-			Result = EDataValidationResult::Invalid;
-			ValidationErrors.Add(FText::Format(DefaultComboNonComboWarning, Args));
-		}
-	}
-	// Input Action contains combo trigger(s)
-	if (bInputActionContainsComboTrigger)
-	{
-		Args.Add("ComboTriggerLocation", LOCTEXT("ComboInInputActionText", "From the Input Action"));
-		// Context contains non-combo trigger(s)
-		if (bContextContainsNonComboTrigger)
-		{
-			Result = EDataValidationResult::Invalid;
-			Args.Add("NonComboTriggerLocation", LOCTEXT("NonComboInContextText", "From the Mapping Context"));
-			ValidationErrors.Add(FText::Format(DefaultComboNonComboWarning, Args));
-		}
-		// Input Action contains non-combo trigger(s) 
-		if (bInputActionContainsNonComboTrigger)
-		{
-			Args.Add("NonComboTriggerLocation", LOCTEXT("NonComboInInputActionText", "From the Input Action"));
-			Result = EDataValidationResult::Invalid;
-			ValidationErrors.Add(FText::Format(DefaultComboNonComboWarning, Args));
+			Args.Add("ComboTriggerLocation", LOCTEXT("ComboInInputActionText", "From the Input Action"));
+			// Context contains non-combo trigger(s)
+			if (bContextContainsNonComboTrigger)
+			{
+				Result = EDataValidationResult::Invalid;
+				Args.Add("NonComboTriggerLocation", LOCTEXT("NonComboInContextText", "From the Mapping Context"));
+				ValidationErrors.Add(FText::Format(DefaultComboNonComboWarning, Args));
+			}
+			// Input Action contains non-combo trigger(s) 
+			if (bInputActionContainsNonComboTrigger)
+			{
+				Args.Add("NonComboTriggerLocation", LOCTEXT("NonComboInInputActionText", "From the Input Action"));
+				Result = EDataValidationResult::Invalid;
+				ValidationErrors.Add(FText::Format(DefaultComboNonComboWarning, Args));
+			}
 		}
 	}
 	
