Fix a crash when calling GetBoundActionValue with a null actor. This can happen if you place the node in an animation blueprint graph and compile it

diff --git a/Engine/Plugins/EnhancedInput/Source/EnhancedInput/Private/EnhancedInputComponent.cpp b/Engine/Plugins/EnhancedInput/Source/EnhancedInput/Private/EnhancedInputComponent.cpp
--- a/Engine/Plugins/EnhancedInput/Source/EnhancedInput/Private/EnhancedInputComponent.cpp
+++ b/Engine/Plugins/EnhancedInput/Source/EnhancedInput/Private/EnhancedInputComponent.cpp
@@ -107,7 +107,7 @@ FInputActionValue UEnhancedInputComponent::GetBoundActionValue(const UInputActio
 			return Binding.GetValue();
 		}
 	}
-	return FInputActionValue(Action->ValueType, FVector::ZeroVector);
+	return FInputActionValue(Action ? Action->ValueType : EInputActionValueType::Axis3D, FVector::ZeroVector);
 }
 
 // Must be in C++ to avoid duplicate statics across execution units
diff --git a/Engine/Plugins/EnhancedInput/Source/EnhancedInput/Private/EnhancedInputLibrary.cpp b/Engine/Plugins/EnhancedInput/Source/EnhancedInput/Private/EnhancedInputLibrary.cpp
--- a/Engine/Plugins/EnhancedInput/Source/EnhancedInput/Private/EnhancedInputLibrary.cpp
+++ b/Engine/Plugins/EnhancedInput/Source/EnhancedInput/Private/EnhancedInputLibrary.cpp
@@ -8,6 +8,7 @@
 #include "EnhancedInputSubsystems.h"
 #include "UObject/UObjectIterator.h"
 #include "GameFramework/Actor.h"
+#include "EnhancedInputModule.h"	// For LogEnhancedInput
 
 #include UE_INLINE_GENERATED_CPP_BY_NAME(EnhancedInputLibrary)
 
