Renaming some world subsystem functions to be more descriptive and designer friendly

#rnx

diff --git a/Engine/Plugins/EnhancedInput/Source/EnhancedInput/Private/EnhancedInputSubsystems.cpp b/Engine/Plugins/EnhancedInput/Source/EnhancedInput/Private/EnhancedInputSubsystems.cpp
--- a/Engine/Plugins/EnhancedInput/Source/EnhancedInput/Private/EnhancedInputSubsystems.cpp
+++ b/Engine/Plugins/EnhancedInput/Source/EnhancedInput/Private/EnhancedInputSubsystems.cpp
@@ -147,7 +147,7 @@ UEnhancedPlayerInput* UEnhancedInputWorldSubsystem::GetPlayerInput() const
 	return PlayerInput;
 }
 
-void UEnhancedInputWorldSubsystem::PushInputComponent(AActor* Actor)
+void UEnhancedInputWorldSubsystem::AddActorInputComponent(AActor* Actor)
 {
 	if (Actor)
 	{
@@ -182,7 +182,7 @@ void UEnhancedInputWorldSubsystem::PushInputComponent(AActor* Actor)
 	}
 }
 
-bool UEnhancedInputWorldSubsystem::PopInputComponent(AActor* Actor)
+bool UEnhancedInputWorldSubsystem::RemoveActorInputComponent(AActor* Actor)
 {
 	if (Actor) 
 	{
diff --git a/Engine/Plugins/EnhancedInput/Source/EnhancedInput/Public/EnhancedInputSubsystems.h b/Engine/Plugins/EnhancedInput/Source/EnhancedInput/Public/EnhancedInputSubsystems.h
--- a/Engine/Plugins/EnhancedInput/Source/EnhancedInput/Public/EnhancedInputSubsystems.h
+++ b/Engine/Plugins/EnhancedInput/Source/EnhancedInput/Public/EnhancedInputSubsystems.h
@@ -70,13 +70,13 @@ public:
 	virtual UEnhancedPlayerInput* GetPlayerInput() const override;
 	//~ End IEnhancedInputSubsystemInterface
 
-	/** Pushes this input component onto the stack to be processed by this subsystem's tick function */
+	/** Adds this Actor's input component onto the stack to be processed by this subsystem's tick function */
 	UFUNCTION(BlueprintCallable, Category = "Input|World")
-	void PushInputComponent(AActor* Actor);
+	void AddActorInputComponent(AActor* Actor);
 
-	/** Removes this input component onto the stack to be processed by this subsystem's tick function */
+	/** Removes this Actor's input component from the stack to be processed by this subsystem's tick function */
 	UFUNCTION(BlueprintCallable, Category = "Input|World")
-	bool PopInputComponent(AActor* Actor);
+	bool RemoveActorInputComponent(AActor* Actor);
 	
 	/** Start the consumption of input messages in this subsystem. This is required to have any Input Action delegates be fired. */
 	UFUNCTION(BlueprintCallable, Category = "Input|World")
