Unshelved from pending changelist '23757394':

Re-organize the Enhanced Input developer settings to have better padding and clean up the categories

#rnx

diff --git a/Engine/Plugins/EnhancedInput/Source/EnhancedInput/Private/EnhancedInputDeveloperSettings.cpp b/Engine/Plugins/EnhancedInput/Source/EnhancedInput/Private/EnhancedInputDeveloperSettings.cpp
--- a/Engine/Plugins/EnhancedInput/Source/EnhancedInput/Private/EnhancedInputDeveloperSettings.cpp
+++ b/Engine/Plugins/EnhancedInput/Source/EnhancedInput/Private/EnhancedInputDeveloperSettings.cpp
@@ -2,6 +2,7 @@
 
 #include "EnhancedInputDeveloperSettings.h"
 #include "EnhancedInputPlatformSettings.h"
+#include "EnhancedPlayerInput.h"
 #include "HAL/IConsoleManager.h"
 
 #include UE_INLINE_GENERATED_CPP_BY_NAME(EnhancedInputDeveloperSettings)
@@ -20,7 +21,12 @@ namespace UE::EnhancedInput::Private
 }
 
 UEnhancedInputDeveloperSettings::UEnhancedInputDeveloperSettings(const FObjectInitializer& Initializer)
-	: Super(Initializer)
+	: Super(Initializer)	
+	, DefaultWorldInputClass(UEnhancedPlayerInput::StaticClass())
+	, bEnableDefaultMappingContexts(true)
+	, bShouldOnlyTriggerLastActionInChord(true)
+	, bEnableWorldSubsystem(true)
+	, bShouldLogAllWorldSubsystemInputs(false)
 {
 	PlatformSettings.Initialize(UEnhancedInputPlatformSettings::StaticClass());
-}
+}
\ No newline at end of file
diff --git a/Engine/Plugins/EnhancedInput/Source/EnhancedInput/Public/EnhancedInputDeveloperSettings.h b/Engine/Plugins/EnhancedInput/Source/EnhancedInput/Public/EnhancedInputDeveloperSettings.h
--- a/Engine/Plugins/EnhancedInput/Source/EnhancedInput/Public/EnhancedInputDeveloperSettings.h
+++ b/Engine/Plugins/EnhancedInput/Source/EnhancedInput/Public/EnhancedInputDeveloperSettings.h
@@ -23,6 +23,8 @@ struct FDefaultContextSetting
 	int32 Priority = 0;
 };
 
+class UEnhancedPlayerInput;
+
 /** Developer settings for Enhanced Input */
 UCLASS(config = Input, defaultconfig, meta = (DisplayName = "Enhanced Input"))
 class ENHANCEDINPUT_API UEnhancedInputDeveloperSettings : public UDeveloperSettingsBackedByCVars
@@ -32,10 +34,6 @@ public:
 	
 	UEnhancedInputDeveloperSettings(const FObjectInitializer& Initializer);
 
-	/** If true, then the DefaultMappingContexts will be applied to all Enhanced Player Inputs. */
-	UPROPERTY(EditAnywhere, Category = "Enhanced Input", meta=(ConsoleVariable="EnhancedInput.EnableDefaultMappingContexts"))
-	bool bEnableDefaultMappingContexts = true;
-	
 	/**
 	 * Array of any input mapping contexts that you want to be applied by default to the Enhanced Input local player subsystem.
 	 * NOTE: These mapping context's can only be from your game's root content directory, not plugins.
@@ -47,12 +45,12 @@ public:
 	 * Array of any input mapping contexts that you want to be applied by default to the Enhanced Input world subsystem.
 	 * NOTE: These mapping context's can only be from your game's root content directory, not plugins.
 	 */
-	UPROPERTY(config, EditAnywhere, Category = "Enhanced Input", meta = (editCondition = "bEnableDefaultMappingContexts"))
+	UPROPERTY(config, EditAnywhere, Category = "Enhanced Input|World Subsystem", meta = (editCondition = "bEnableDefaultMappingContexts && bEnableWorldSubsystem"))
 	TArray<FDefaultContextSetting> DefaultWorldSubsystemMappingContexts;
 	
 	/** The default player input class that the Enhanced Input world subsystem will use. */
-	UPROPERTY(config, EditAnywhere, NoClear, Category = "Enhanced Input")
-	TSoftClassPtr<class UEnhancedPlayerInput> DefaultWorldInputClass;
+	UPROPERTY(config, EditAnywhere, NoClear, Category = "Enhanced Input|World Subsystem", meta=(editCondition = "bEnableWorldSubsystem"))
+	TSoftClassPtr<UEnhancedPlayerInput> DefaultWorldInputClass;
 
 	/**
 	 * Platform specific settings for Enhanced Input.
@@ -61,6 +59,10 @@ public:
 	UPROPERTY(EditAnywhere, Category = "Enhanced Input")
 	FPerPlatformSettings PlatformSettings;
 
+	/** If true, then the DefaultMappingContexts will be applied to all Enhanced Input Subsystems. */
+	UPROPERTY(EditAnywhere, Category = "Enhanced Input", meta = (ConsoleVariable = "EnhancedInput.EnableDefaultMappingContexts"))
+	uint8 bEnableDefaultMappingContexts : 1;
+
 	/**
 	 * If true, then only the last action in a ChordedAction trigger will be fired.
 	 * This means that only the action that has the ChordedTrigger on it will be fired, not the individual steps.
@@ -68,17 +70,17 @@ public:
 	 * Default value is true.
 	 */
 	UPROPERTY(EditAnywhere, Category = "Enhanced Input", meta=(ConsoleVariable="EnhancedInput.OnlyTriggerLastActionInChord"))
-	bool bShouldOnlyTriggerLastActionInChord = true;
+	uint8 bShouldOnlyTriggerLastActionInChord : 1;
 	
 	/** If true, then the world subsystem will be created. */
-	UPROPERTY(config, EditAnywhere, Category = "Enhanced Input")
-	bool bEnableWorldSubsystem = true;
+	UPROPERTY(config, EditAnywhere, Category = "Enhanced Input|World Subsystem")
+	uint8 bEnableWorldSubsystem : 1;
 	
 	/**
  	 * If true then the Enhanced Input world subsystem will log all input that is being processed by it (keypresses, analog values, etc)
- 	 * Note: This can produce A LOT of logs, so only use this if you are debugging something.
+ 	 * Note: This can produce A LOT of logs, so only use this if you are debugging something. Does nothing in shipping builds
  	 */
-	UPROPERTY(config, EditAnywhere, Category = "Enhanced Input", meta=(ConsoleVariable="EnhancedInput.bShouldLogAllWorldSubsystemInputs"))
+	UPROPERTY(config, EditAnywhere, Category = "Enhanced Input|World Subsystem", meta=(editCondition = "bEnableWorldSubsystem", ConsoleVariable="EnhancedInput.bShouldLogAllWorldSubsystemInputs"))
 	uint8 bShouldLogAllWorldSubsystemInputs : 1;
 };
 
