Update the Enhanced Input world subsystem to pass along the input device ID. Remove the "World" property on this because that wouldn't work if there are multiple world subsystems, like during level transitions.

#rnx
#pf skip

diff --git a/Engine/Plugins/EnhancedInput/Source/EnhancedInput/Private/EnhancedInputWorldProcessor.cpp b/Engine/Plugins/EnhancedInput/Source/EnhancedInput/Private/EnhancedInputWorldProcessor.cpp
--- a/Engine/Plugins/EnhancedInput/Source/EnhancedInput/Private/EnhancedInputWorldProcessor.cpp
+++ b/Engine/Plugins/EnhancedInput/Source/EnhancedInput/Private/EnhancedInputWorldProcessor.cpp
@@ -4,6 +4,9 @@
 #include "EnhancedInputSubsystems.h"
 #include "Engine/World.h"
 #include "Framework/Application/SlateApplication.h"
+#include "EnhancedInputDeveloperSettings.h"
+#include "UObject/UObjectIterator.h"
+#include "GenericPlatform/GenericPlatformInputDeviceMapper.h"
 
 void FEnhancedInputWorldProcessor::Tick(const float DeltaTime, FSlateApplication& SlateApp, TSharedRef<ICursor> Cursor)
 {
@@ -18,6 +21,7 @@ bool FEnhancedInputWorldProcessor::HandleKeyDownEvent(FSlateApplication& SlateAp
 	Params.Delta.X = 1.0;
 	Params.DeltaTime = SlateApp.GetDeltaTime();
 	Params.NumSamples = Params.Key.IsAnalog() ? 1 : 0;
+	Params.InputDevice = InKeyEvent.GetInputDeviceId();
 	
 	InputKeyToSubsystem(Params);
 	
@@ -32,6 +36,7 @@ bool FEnhancedInputWorldProcessor::HandleKeyUpEvent(FSlateApplication& SlateApp,
 	Params.Delta.X = 0.0;
 	Params.DeltaTime = SlateApp.GetDeltaTime();
 	Params.NumSamples = Params.Key.IsAnalog() ? 1 : 0;
+	Params.InputDevice = InKeyEvent.GetInputDeviceId();
 	
 	InputKeyToSubsystem(Params);
 	return IInputProcessor::HandleKeyUpEvent(SlateApp, InKeyEvent);
@@ -45,6 +50,7 @@ bool FEnhancedInputWorldProcessor::HandleAnalogInputEvent(FSlateApplication& Sla
 	Params.Delta.X = InAnalogInputEvent.GetAnalogValue();
 	Params.DeltaTime = SlateApp.GetDeltaTime();
 	Params.NumSamples = 1;
+	Params.InputDevice = InAnalogInputEvent.GetInputDeviceId();
 	
 	InputKeyToSubsystem(Params);
 	return IInputProcessor::HandleAnalogInputEvent(SlateApp, InAnalogInputEvent);
@@ -65,6 +71,7 @@ bool FEnhancedInputWorldProcessor::HandleMouseButtonDownEvent(FSlateApplication&
 	Params.Delta.X = 1.0;
 	Params.DeltaTime = SlateApp.GetDeltaTime();
 	Params.NumSamples = 0;
+	Params.InputDevice = MouseEvent.GetInputDeviceId();
 	
 	InputKeyToSubsystem(Params);
 	return IInputProcessor::HandleMouseButtonDownEvent(SlateApp, MouseEvent);
@@ -78,6 +85,7 @@ bool FEnhancedInputWorldProcessor::HandleMouseButtonUpEvent(FSlateApplication& S
 	Params.Delta.X = 0.0;
 	Params.DeltaTime = SlateApp.GetDeltaTime();
 	Params.NumSamples = 0;
+	Params.InputDevice = MouseEvent.GetInputDeviceId();
 	
 	InputKeyToSubsystem(Params);
 	return IInputProcessor::HandleMouseButtonUpEvent(SlateApp, MouseEvent);
@@ -91,6 +99,7 @@ bool FEnhancedInputWorldProcessor::HandleMouseButtonDoubleClickEvent(FSlateAppli
 	Params.Delta.X = 1.0;
 	Params.DeltaTime = SlateApp.GetDeltaTime();
 	Params.NumSamples = 0;
+	Params.InputDevice = MouseEvent.GetInputDeviceId();
 	
 	InputKeyToSubsystem(Params);
 	return IInputProcessor::HandleMouseButtonDoubleClickEvent(SlateApp, MouseEvent);
@@ -109,6 +118,7 @@ bool FEnhancedInputWorldProcessor::HandleMouseWheelOrGestureEvent(FSlateApplicat
 		PressedParams.Delta.X = 1.0;
 		PressedParams.DeltaTime = SlateApp.GetDeltaTime();
 		PressedParams.NumSamples = 0;
+		PressedParams.InputDevice = InWheelEvent.GetInputDeviceId();
 		
 		FInputKeyParams ReleasedParams = PressedParams;
 		ReleasedParams.Event = IE_Released;
@@ -123,6 +133,7 @@ bool FEnhancedInputWorldProcessor::HandleMouseWheelOrGestureEvent(FSlateApplicat
 		Params.Delta.X = InWheelEvent.GetWheelDelta();
 		Params.DeltaTime = SlateApp.GetDeltaTime();
 		Params.NumSamples = 1;
+		Params.InputDevice = InWheelEvent.GetInputDeviceId();
 		
 		InputKeyToSubsystem(Params);
 	}
@@ -147,6 +158,8 @@ void FEnhancedInputWorldProcessor::ProcessAccumulatedPointerInput(float DeltaTim
 		Params.Delta.X = CachedCursorDelta.X;
 		Params.DeltaTime = DeltaTime;
 		Params.NumSamples = NumCursorSamplesThisFrame.X;
+		Params.InputDevice = IPlatformInputDeviceMapper::Get().GetDefaultInputDevice();
+
 		InputKeyToSubsystem(Params);
 	}
 
@@ -157,6 +170,8 @@ void FEnhancedInputWorldProcessor::ProcessAccumulatedPointerInput(float DeltaTim
 		Params.Delta.X = CachedCursorDelta.Y;
 		Params.DeltaTime = DeltaTime;
 		Params.NumSamples = NumCursorSamplesThisFrame.Y;
+		Params.InputDevice = IPlatformInputDeviceMapper::Get().GetDefaultInputDevice();
+
 		InputKeyToSubsystem(Params);
 	}
 	
@@ -167,11 +182,12 @@ void FEnhancedInputWorldProcessor::ProcessAccumulatedPointerInput(float DeltaTim
 
 bool FEnhancedInputWorldProcessor::InputKeyToSubsystem(const FInputKeyParams& Params)
 {
-	if (WorldFromSubsystem)
+	if (GetDefault<UEnhancedInputDeveloperSettings>()->bEnableWorldSubsystem)
 	{
-		if (UEnhancedInputWorldSubsystem* WorldSubsystem = UWorld::GetSubsystem<UEnhancedInputWorldSubsystem>(WorldFromSubsystem))
+		// Tell all the world subsystems about the key that has been pressed
+		for (TObjectIterator<UEnhancedInputWorldSubsystem> It; It; ++It)
 		{
-			return WorldSubsystem->InputKey(Params);
+			return (*It)->InputKey(Params);
 		}
 	}
 	
diff --git a/Engine/Plugins/EnhancedInput/Source/EnhancedInput/Public/EnhancedInputWorldProcessor.h b/Engine/Plugins/EnhancedInput/Source/EnhancedInput/Public/EnhancedInputWorldProcessor.h
--- a/Engine/Plugins/EnhancedInput/Source/EnhancedInput/Public/EnhancedInputWorldProcessor.h
+++ b/Engine/Plugins/EnhancedInput/Source/EnhancedInput/Public/EnhancedInputWorldProcessor.h
@@ -31,11 +31,6 @@ public:
 	virtual const TCHAR* GetDebugName() const { return TEXT("EnhancedInputWorldPreprocessor"); }
 	//~ End IInputProcessor interface
 
-	// Setter for the the world variable
-	void SetSubsystemWorld(UWorld* World) { WorldFromSubsystem = World; }
-	// Getter for the the world variable
-	UWorld* GetSubsystemWorld() const { return WorldFromSubsystem; }
-
 protected:
 
 	void UpdateCachedPointerPosition(const FPointerEvent& MouseEvent);
@@ -54,10 +49,6 @@ protected:
 	
 	/** The value of the mouse cursor from the most recent mouse event */
 	FVector2D CachedCursorDelta = FVector2D::ZeroVector;
-
-private:
-	UWorld* WorldFromSubsystem = nullptr;
-
 };
 
 #if UE_ENABLE_INCLUDE_ORDER_DEPRECATED_IN_5_2
