Unshelved from pending changelist '23756918':

- Change the Enhanced Input World Subsystem from a tickable game object to a normal world subsystem. This way we can control when the system is ticked and keep it aligned with the local player input subsystem better. This ensures some more stable results when processing input.
- Add some TRACE_CPUPROFILER_EVENT_SCOPE to the EI module so that we can better track its performance
- Remove the "bIsConsumingInput" because we dont need it, as it should always consume input.
- Make the InputKey function protected instead of public because it should only ever be called from the input preprocessor

diff --git a/Engine/Plugins/EnhancedInput/Source/EnhancedInput/Private/EnhancedInputLibrary.cpp b/Engine/Plugins/EnhancedInput/Source/EnhancedInput/Private/EnhancedInputLibrary.cpp
--- a/Engine/Plugins/EnhancedInput/Source/EnhancedInput/Private/EnhancedInputLibrary.cpp
+++ b/Engine/Plugins/EnhancedInput/Source/EnhancedInput/Private/EnhancedInputLibrary.cpp
@@ -9,13 +9,20 @@
 #include "UObject/UObjectIterator.h"
 #include "GameFramework/Actor.h"
 #include "EnhancedInputModule.h"	// For LogEnhancedInput
+#include "EnhancedInputDeveloperSettings.h"
 
 #include UE_INLINE_GENERATED_CPP_BY_NAME(EnhancedInputLibrary)
 
 void UEnhancedInputLibrary::ForEachSubsystem(TFunctionRef<void(IEnhancedInputSubsystemInterface*)> SubsystemPredicate)
 {
-	// TODO: World subsystem for enhanced input, so that you bind to actions without an owning player controller
-	// This is useful for widgets, main menu situations, or just stuff that isn't dependant on a player.
+	// World subsystem
+	if (GetDefault<UEnhancedInputDeveloperSettings>()->bEnableWorldSubsystem)
+	{
+		for (TObjectIterator<UEnhancedInputWorldSubsystem> It; It; ++It)
+		{
+			SubsystemPredicate(Cast<IEnhancedInputSubsystemInterface>(*It));
+		}
+	}
 
 	// Players
 	for (TObjectIterator<UEnhancedInputLocalPlayerSubsystem> It; It; ++It)
diff --git a/Engine/Plugins/EnhancedInput/Source/EnhancedInput/Private/EnhancedInputModule.cpp b/Engine/Plugins/EnhancedInput/Source/EnhancedInput/Private/EnhancedInputModule.cpp
--- a/Engine/Plugins/EnhancedInput/Source/EnhancedInput/Private/EnhancedInputModule.cpp
+++ b/Engine/Plugins/EnhancedInput/Source/EnhancedInput/Private/EnhancedInputModule.cpp
@@ -12,6 +12,7 @@
 #include "GameFramework/PlayerController.h"
 #include "UObject/Package.h"
 #include "UObject/UObjectIterator.h"
+#include "EnhancedInputDeveloperSettings.h"
 
 #define LOCTEXT_NAMESPACE "EnhancedInput"
 
@@ -284,6 +285,8 @@ void FEnhancedInputModule::ShutdownModule()
 
 void FEnhancedInputModule::Tick(float DeltaTime)
 {
+	TRACE_CPUPROFILER_EVENT_SCOPE(FEnhancedInputModule::Tick);
+
 	// This may tick multiple times per frame? See MIDIDevice module.
 	if (LastFrameNumberWeTicked == GFrameCounter)
 	{
@@ -291,6 +294,20 @@ void FEnhancedInputModule::Tick(float DeltaTime)
 	}
 	LastFrameNumberWeTicked = GFrameCounter;
 
+	/* 
+	* Tick the World subsystems 
+	* The Enhanced Input local player subsystem will have their Player Input's ticked by their owning 
+	* Player Controller in APlayerController::TickPlayerInput, but because the world subsystem has no
+	* owning controller we need to tick it here.
+	*/
+	if (GetDefault<UEnhancedInputDeveloperSettings>()->bEnableWorldSubsystem)
+	{
+		for (TObjectIterator<UEnhancedInputWorldSubsystem> It; It; ++It)
+		{
+			(*It)->TickPlayerInput(DeltaTime);
+		}
+	}	
+
 	UEnhancedInputLibrary::ForEachSubsystem([DeltaTime](IEnhancedInputSubsystemInterface* Subsystem)
 		{
 			Subsystem->RebuildControlMappings();
diff --git a/Engine/Plugins/EnhancedInput/Source/EnhancedInput/Private/EnhancedInputSubsystems.cpp b/Engine/Plugins/EnhancedInput/Source/EnhancedInput/Private/EnhancedInputSubsystems.cpp
--- a/Engine/Plugins/EnhancedInput/Source/EnhancedInput/Private/EnhancedInputSubsystems.cpp
+++ b/Engine/Plugins/EnhancedInput/Source/EnhancedInput/Private/EnhancedInputSubsystems.cpp
@@ -42,6 +42,23 @@ void UEnhancedInputLocalPlayerSubsystem::ControlMappingsRebuiltThisFrame()
 // *
 // **************************************************************************************************
 
+bool UEnhancedInputWorldSubsystem::ShouldCreateSubsystem(UObject* Outer) const
+{
+	if (!FSlateApplication::IsInitialized())
+	{
+		return false;
+	}
+
+	// Getting setting on whether to turn off subsystem or not
+	const bool bShouldCreate = GetDefault<UEnhancedInputDeveloperSettings>()->bEnableWorldSubsystem;
+	if (!bShouldCreate)
+	{
+		UE_LOG(LogWorldSubsystemInput, Log, TEXT("UEnhancedInputDeveloperSettings::bEnableWorldSubsystem is false, the world subsystem will not be created!"));
+	}
+
+	return bShouldCreate && Super::ShouldCreateSubsystem(Outer);
+}
+
 void UEnhancedInputWorldSubsystem::Initialize(FSubsystemCollectionBase& Collection)
 {
 	Super::Initialize(Collection);
@@ -59,14 +76,15 @@ void UEnhancedInputWorldSubsystem::Initialize(FSubsystemCollectionBase& Collecti
 	PlayerInput = NewObject<UEnhancedPlayerInput>(this, PlayerInputClass, TEXT("EIWorldSubsystem_PlayerInput0"));	
 	ensureMsgf(PlayerInput, TEXT("UEnhancedInputWorldSubsystem::Initialize failed to create PlayerInput! This subsystem will not tick!"));
 	
-	if (FSlateApplication::IsInitialized())
+	if (ensure(FSlateApplication::IsInitialized()))
 	{
 		InputPreprocessor = MakeShared<FEnhancedInputWorldProcessor>();
-		InputPreprocessor->SetSubsystemWorld(GetWorld());
-		FSlateApplication::Get().RegisterInputPreProcessor(InputPreprocessor, 0);
+		FSlateApplication::Get().RegisterInputPreProcessor(InputPreprocessor);
 	}
 	
-	StartConsumingInput();
+	// Add any default mapping contexts
+	AddDefaultMappingContexts();
+	RequestRebuildControlMappings();
 }
 
 void UEnhancedInputWorldSubsystem::Deinitialize()
@@ -81,47 +99,26 @@ void UEnhancedInputWorldSubsystem::Deinitialize()
 	InputPreprocessor.Reset();
 	CurrentInputStack.Empty();
 	PlayerInput = nullptr;
-}
 
-bool UEnhancedInputWorldSubsystem::ShouldCreateSubsystem(UObject* Outer) const
-{
-	if (!FSlateApplication::IsInitialized())
-	{
-		return false;
-	}
-	
-	// Getting setting on whether to turn off subsystem or not
-	const bool bShouldCreate = GetDefault<UEnhancedInputDeveloperSettings>()->bEnableWorldSubsystem;
-	if (!bShouldCreate)
-	{
-		UE_LOG(LogWorldSubsystemInput, Log, TEXT("UEnhancedInputDeveloperSettings::bEnableWorldSubsystem is false, the world subsystem will not be created!"));
-	}
-	
-	return bShouldCreate && Super::ShouldCreateSubsystem(Outer);
+	RemoveDefaultMappingContexts();
 }
 
-ETickableTickType UEnhancedInputWorldSubsystem::GetTickableTickType() const
+bool UEnhancedInputWorldSubsystem::DoesSupportWorldType(const EWorldType::Type WorldType) const
 {
-	return IsTemplate() ? ETickableTickType::Never : ETickableTickType::Always;
+	// The world subsystem shouldn't be used in the editor
+	return WorldType == EWorldType::Game || WorldType == EWorldType::PIE;
 }
 
-void UEnhancedInputWorldSubsystem::Tick(float DeltaTime)
+void UEnhancedInputWorldSubsystem::TickPlayerInput(float DeltaTime)
 {
-	if (!ensure(PlayerInput))
-	{
-		return;
-	}
-
-	FModifyContextOptions Options;
-	Options.bForceImmediately = true;
+	TRACE_CPUPROFILER_EVENT_SCOPE(UEnhancedInputWorldSubsystem::TickPlayerInput);
 	
-	// Rebuild the control mappings and tick any forced input that may have been injected
-	RequestRebuildControlMappings(Options);
-	TickForcedInput(DeltaTime);
-
 	static TArray<UInputComponent*> InputStack;
+	InputStack.Reset();
+
 	// Build the input stack
 	{
+		TRACE_CPUPROFILER_EVENT_SCOPE(UEnhancedInputWorldSubsystem::TickPlayerInput::BuildInputStack);
 		for (int32 i = 0; i < CurrentInputStack.Num(); ++i)
 		{
 			if (UInputComponent* IC = CurrentInputStack[i].Get())
@@ -137,14 +134,7 @@ void UEnhancedInputWorldSubsystem::Tick(float DeltaTime)
 	
 	// Process input stack on the player input
 	PlayerInput->Tick(DeltaTime);
-	PlayerInput->ProcessInputStack(InputStack, DeltaTime, GetWorld()->IsPaused());
-
-	InputStack.Reset();
-}
-
-UEnhancedPlayerInput* UEnhancedInputWorldSubsystem::GetPlayerInput() const
-{
-	return PlayerInput;
+	PlayerInput->ProcessInputStack(InputStack, DeltaTime, GetWorld()->IsPaused());	
 }
 
 void UEnhancedInputWorldSubsystem::AddActorInputComponent(AActor* Actor)
@@ -199,37 +189,19 @@ bool UEnhancedInputWorldSubsystem::RemoveActorInputComponent(AActor* Actor)
     return false;
 }
 
-void UEnhancedInputWorldSubsystem::StartConsumingInput()
-{
-	bIsCurrentlyConsumingInput = true;
-	
-	AddDefaultMappingContexts();
-
-	RequestRebuildControlMappings();
-}
-
-void UEnhancedInputWorldSubsystem::StopConsumingInput()
-{
-	bIsCurrentlyConsumingInput = false;
-
-	RemoveDefaultMappingContexts();
-
-	RequestRebuildControlMappings();
-}
-
 bool UEnhancedInputWorldSubsystem::InputKey(const FInputKeyParams& Params)
 {
 	if (PlayerInput)
 	{
-		if (bIsCurrentlyConsumingInput)
+
+#if !UE_BUILD_SHIPPING
+		if (GetDefault<UEnhancedInputDeveloperSettings>()->bShouldLogAllWorldSubsystemInputs)
 		{
-			if (GetDefault<UEnhancedInputDeveloperSettings>()->bShouldLogAllWorldSubsystemInputs)
-			{
-				UE_LOG(LogWorldSubsystemInput, Log, TEXT("EI %s World Subsystem InputKey : [%s]"), LexToString(GetWorld()->WorldType), *Params.Key.ToString());
-			}
-			return PlayerInput->InputKey(Params);
+			UE_LOG(LogWorldSubsystemInput, VeryVerbose, TEXT("EI %s World Subsystem InputKey : [%s]"), LexToString(GetWorld()->WorldType), *Params.Key.ToString());
 		}
-		return false;
+#endif
+
+		return PlayerInput->InputKey(Params);
 	}
 
 	ensureAlwaysMsgf(false, TEXT("Attempting to input a key to the EnhancedInputWorldSubsystem, but there is no Player Input!"));
@@ -267,6 +239,11 @@ void UEnhancedInputWorldSubsystem::RemoveDefaultMappingContexts()
 	}
 }
 
+UEnhancedPlayerInput* UEnhancedInputWorldSubsystem::GetPlayerInput() const
+{
+	return PlayerInput;
+}
+
 void UEnhancedInputWorldSubsystem::ShowDebugInfo(UCanvas* Canvas)
 {
 	if (!Canvas)
diff --git a/Engine/Plugins/EnhancedInput/Source/EnhancedInput/Public/EnhancedInputSubsystems.h b/Engine/Plugins/EnhancedInput/Source/EnhancedInput/Public/EnhancedInputSubsystems.h
--- a/Engine/Plugins/EnhancedInput/Source/EnhancedInput/Public/EnhancedInputSubsystems.h
+++ b/Engine/Plugins/EnhancedInput/Source/EnhancedInput/Public/EnhancedInputSubsystems.h
@@ -45,62 +45,66 @@ public:
  * 
  * Note: if you do have an actor with an owning Player Controller use the local player input subsystem instead.
  */
-UCLASS()
-class ENHANCEDINPUT_API UEnhancedInputWorldSubsystem : public UTickableWorldSubsystem, public IEnhancedInputSubsystemInterface
+UCLASS(DisplayName="Enhanced Input World Subsystem (Experimental)")
+class ENHANCEDINPUT_API UEnhancedInputWorldSubsystem : public UWorldSubsystem, public IEnhancedInputSubsystemInterface
 {
+
+// The Enhanced Input Module ticks the player input on this subsystem
+friend class FEnhancedInputModule;
+// The input processor tells us about what keys are pressed
+friend class FEnhancedInputWorldProcessor;
+
 	GENERATED_BODY()
 
 public:
 
-	//~ Begin USubsystem interface
+	//~ Begin UWorldSubsystem interface
+	virtual bool ShouldCreateSubsystem(UObject* Outer) const override;
 	virtual void Initialize(FSubsystemCollectionBase& Collection) override;
 	virtual void Deinitialize() override;
-	virtual bool ShouldCreateSubsystem(UObject* Outer) const override;
-	//~ End USubsystem interface
-
-	//~ Begin FTickableGameObject interface
-	//virtual UWorld* GetTickableGameObjectWorld() const override;
-	virtual bool IsTickableInEditor() const { return true; }
-	virtual ETickableTickType GetTickableTickType() const override;
-	virtual void Tick(float DeltaTime) override;
-	TStatId GetStatId() const override { RETURN_QUICK_DECLARE_CYCLE_STAT(UEnhancedInputWorldSubsystem, STATGROUP_Tickables); }
-	//~ End FTickableGameObject interface
-	
+protected:
+	virtual bool DoesSupportWorldType(const EWorldType::Type WorldType) const;
+	//~ End UWorldSubsystem interface
+
+public:	
 	//~ Begin IEnhancedInputSubsystemInterface
 	virtual UEnhancedPlayerInput* GetPlayerInput() const override;
+	virtual void ShowDebugInfo(UCanvas* Canvas) override;
 	//~ End IEnhancedInputSubsystemInterface
 
 	/** Adds this Actor's input component onto the stack to be processed by this subsystem's tick function */
-	UFUNCTION(BlueprintCallable, Category = "Input|World")
+	UFUNCTION(BlueprintCallable, Category = "Input|World", meta=(DefaultToSelf = "Actor"))
 	void AddActorInputComponent(AActor* Actor);
 
 	/** Removes this Actor's input component from the stack to be processed by this subsystem's tick function */
-	UFUNCTION(BlueprintCallable, Category = "Input|World")
+	UFUNCTION(BlueprintCallable, Category = "Input|World", meta = (DefaultToSelf = "Actor"))
 	bool RemoveActorInputComponent(AActor* Actor);
 	
-	/** Start the consumption of input messages in this subsystem. This is required to have any Input Action delegates be fired. */
-	UFUNCTION(BlueprintCallable, Category = "Input|World")
-	void StartConsumingInput();
+protected:
+
+	/** 
+	* Inputs a key on this subsystem's player input which can then be processed as normal during Tick.
+	* 
+	* This should only be called by the FEnhancedInputWorldProcessor 
+	*/
+	bool InputKey(const FInputKeyParams& Params);
+
+	/** 
+	* Builds the current input stack and ticks the world subsystem's player input.
+	* 
+	* Called from the Enhanced Input Module Tick.
+	* 
+	* The Enhanced Input local player subsystem will have their Player Input's ticked by their owning 
+	* Player Controller in APlayerController::TickPlayerInput, but because the world subsystem has no 
+	* owning controller we need to tick it elsewhere.
+	*/
+	void TickPlayerInput(float DeltaTime);
 
-	/** Tells this subsystem to stop ticking and consuming any input. This will stop any Input Action Delegates from being called. */
-	UFUNCTION(BlueprintCallable, Category = "Input|World")
-	void StopConsumingInput();
-	
-	/** Returns true if this subsystem is currently consuming input */
-    UFUNCTION(BlueprintCallable, BlueprintPure, Category = "Input|World")
-    bool IsConsumingInput() const { return bIsCurrentlyConsumingInput; }
-	
-    /** Inputs a key on this subsystem's player input which can then be processed as normal during Tick. */
-    bool InputKey(const FInputKeyParams& Params);
-	
 	/** Adds all the default mapping contexts */
 	void AddDefaultMappingContexts();
 
 	/** Removes all the default mapping contexts */
 	void RemoveDefaultMappingContexts();
-	
-	virtual void ShowDebugInfo(UCanvas* Canvas) override;
-private:
 
 	/** The player input that is processing the input within this subsystem */
 	UPROPERTY()
@@ -109,10 +113,7 @@ private:
 	/**
 	 * Input processor that is created on Initalize.
 	 */
-	TSharedPtr<FEnhancedInputWorldProcessor> InputPreprocessor = nullptr;
-	
-	/** If true, then this subsystem will Tick and process input delegates. */
-	bool bIsCurrentlyConsumingInput = false;
+	TSharedPtr<FEnhancedInputWorldProcessor> InputPreprocessor = nullptr;	
 	
 	/** Internal. This is the current stack of InputComponents that is being processed by the PlayerInput. */
 	UPROPERTY(Transient)
