Name the player mappable options object and slot as experimental. Don't count on long term support for these types as they are likely to change

#rnx

diff --git a/Engine/Plugins/EnhancedInput/Source/EnhancedInput/Public/PlayerMappableKeySettings.h b/Engine/Plugins/EnhancedInput/Source/EnhancedInput/Public/PlayerMappableKeySettings.h
--- a/Engine/Plugins/EnhancedInput/Source/EnhancedInput/Public/PlayerMappableKeySettings.h
+++ b/Engine/Plugins/EnhancedInput/Source/EnhancedInput/Public/PlayerMappableKeySettings.h
@@ -8,8 +8,9 @@ struct FEnhancedActionKeyMapping;
 
 /**
 * Hold setting information of an Action Input or a Action Key Mapping for setting screen and save purposes.
+* Experimental: Do not count on long term support for this structure.
 */
-UCLASS(DefaultToInstanced, EditInlineNew)
+UCLASS(DefaultToInstanced, EditInlineNew, DisplayName="Player Mappable Key Settings (Experimental)")
 class ENHANCEDINPUT_API UPlayerMappableKeySettings : public UObject
 {
 	GENERATED_BODY()
diff --git a/Engine/Plugins/EnhancedInput/Source/EnhancedInput/Public/PlayerMappableKeySlot.h b/Engine/Plugins/EnhancedInput/Source/EnhancedInput/Public/PlayerMappableKeySlot.h
--- a/Engine/Plugins/EnhancedInput/Source/EnhancedInput/Public/PlayerMappableKeySlot.h
+++ b/Engine/Plugins/EnhancedInput/Source/EnhancedInput/Public/PlayerMappableKeySlot.h
@@ -6,8 +6,9 @@
 
 /**
 * Explicitly identifies the slot for a player mappable key
+* Experimental: Do not count on long term support for this structure.
 */
-USTRUCT(BlueprintType)
+USTRUCT(BlueprintType, DisplayName="Player Mappable Key Slot (Experimental)")
 struct ENHANCEDINPUT_API FPlayerMappableKeySlot
 {
 	GENERATED_BODY()
