Adding some data validation for Combo Step Completion and Cancellation States to make sure at least one state is set

diff --git a/Engine/Plugins/EnhancedInput/Source/EnhancedInput/Private/InputTriggers.cpp b/Engine/Plugins/EnhancedInput/Source/EnhancedInput/Private/InputTriggers.cpp
--- a/Engine/Plugins/EnhancedInput/Source/EnhancedInput/Private/InputTriggers.cpp
+++ b/Engine/Plugins/EnhancedInput/Source/EnhancedInput/Private/InputTriggers.cpp
@@ -327,6 +327,26 @@ EDataValidationResult UInputTriggerCombo::IsDataValid(TArray<FText>& ValidationE
 		ValidationErrors.Add(LOCTEXT("NoComboSteps", "There must be at least one combo step in the Combo Trigger!"));
 	}
 
+	// Making sure combo completion states have at least one state
+	for (const FInputComboStepData& ComboStep : ComboActions)
+	{
+		if (ComboStep.ComboStepCompletionStates == 0)
+		{
+			Result = EDataValidationResult::Invalid;
+			ValidationErrors.Add(FText::Format(LOCTEXT("NoCompletionStates", "There must be at least one completion state in ComboStep Completion States in the {0} combo step in order to progress the combo!"), FText::FromString(ComboStep.ComboStepAction.GetName())));
+		}
+	}
+
+	// Making sure cancellation states have at least one state
+	for (const FInputCancelAction& CancelAction : InputCancelActions)
+	{
+		if (CancelAction.CancellationStates == 0)
+		{
+			Result = EDataValidationResult::Invalid;
+			ValidationErrors.Add(FText::Format(LOCTEXT("NoCancellationStates", "There must be at least one cancellation state in Cancellation States in the {0} cancel action in order to cancel the combo!"), FText::FromString(CancelAction.CancelAction.GetName())));
+		}
+	}
+
 	return Result;
 }
 #endif // WITH_EDITOR
