Only attempt to auto upgrade to Enhanced Input if the project has a name. This solves an issue where you would get toasts about EI from the project launcher.

#pf 63e1344b797b029c0a65d67c

diff --git a/Engine/Plugins/EnhancedInput/Source/InputEditor/Private/InputEditorModule.cpp b/Engine/Plugins/EnhancedInput/Source/InputEditor/Private/InputEditorModule.cpp
--- a/Engine/Plugins/EnhancedInput/Source/InputEditor/Private/InputEditorModule.cpp
+++ b/Engine/Plugins/EnhancedInput/Source/InputEditor/Private/InputEditorModule.cpp
@@ -479,7 +479,10 @@ void FInputEditorModule::ShutdownModule()
 
 void FInputEditorModule::OnMainFrameCreationFinished(TSharedPtr<SWindow> InRootWindow, bool bIsRunningStartupDialog)
 {
-	AutoUpgradeDefaultInputClasses();
+	if (FApp::HasProjectName())
+	{
+		AutoUpgradeDefaultInputClasses();
+	}	
 }
 
 namespace UE::Input
