Add "Enhanced Input User Settings".

The "User Settings" is a save game object that provides a place "out of the box" to _store_ player mappable keys. In the future, this is somewhere that we can add Input Accessibility settings, deadzone settings, and other input related settings that a lot of games want right out of the box!

This is a save game instead of just a config because these settings are commonly something that you want stored in the cloud, or are otherwise the same no matter what machine the player is on.

Each UEnhancedPlayerInput will have an instance of this save file, so it will support multiple local users having different settings as well!

By default, this feature is off and needs to be manually turned on in the Enhanced Input developer settings, and it can be disabled via the config files as well.

This CL is simply adding the class and files, the user settings won't even be created with this CL, so there are no behavioral changes at all.

#pf 63e14a5e1466610c647f1e3a

diff --git a/Engine/Plugins/EnhancedInput/Source/EnhancedInput/EnhancedInput.Build.cs b/Engine/Plugins/EnhancedInput/Source/EnhancedInput/EnhancedInput.Build.cs
--- a/Engine/Plugins/EnhancedInput/Source/EnhancedInput/EnhancedInput.Build.cs
+++ b/Engine/Plugins/EnhancedInput/Source/EnhancedInput/EnhancedInput.Build.cs
@@ -16,6 +16,7 @@ namespace UnrealBuildTool.Rules
 					"Slate",
                     "SlateCore",
                     "DeveloperSettings",
+                    "GameplayTags",
                 }
             );
         }
diff --git a/Engine/Plugins/EnhancedInput/Source/EnhancedInput/Private/EnhancedInputDeveloperSettings.cpp b/Engine/Plugins/EnhancedInput/Source/EnhancedInput/Private/EnhancedInputDeveloperSettings.cpp
--- a/Engine/Plugins/EnhancedInput/Source/EnhancedInput/Private/EnhancedInputDeveloperSettings.cpp
+++ b/Engine/Plugins/EnhancedInput/Source/EnhancedInput/Private/EnhancedInputDeveloperSettings.cpp
@@ -2,6 +2,7 @@
 
 #include "EnhancedInputDeveloperSettings.h"
 #include "EnhancedInputPlatformSettings.h"
+#include "UserSettings/EnhancedInputUserSettings.h"
 #include "EnhancedPlayerInput.h"
 #include "HAL/IConsoleManager.h"
 
@@ -21,8 +22,11 @@ namespace UE::EnhancedInput::Private
 }
 
 UEnhancedInputDeveloperSettings::UEnhancedInputDeveloperSettings(const FObjectInitializer& Initializer)
-	: Super(Initializer)	
+	: Super(Initializer)
+	, UserSettingsClass(UEnhancedInputUserSettings::StaticClass())
+	, DefaultPlayerMappableKeyProfileClass(UEnhancedPlayerMappableKeyProfile::StaticClass())
 	, DefaultWorldInputClass(UEnhancedPlayerInput::StaticClass())
+	, bEnableUserSettings(false)
 	, bEnableDefaultMappingContexts(true)
 	, bShouldOnlyTriggerLastActionInChord(true)
 	, bEnableWorldSubsystem(false)
diff --git a/Engine/Plugins/EnhancedInput/Source/EnhancedInput/Private/UserSettings/EnhancedInputUserSettings.cpp b/Engine/Plugins/EnhancedInput/Source/EnhancedInput/Private/UserSettings/EnhancedInputUserSettings.cpp
new file mode 100644
--- /dev/null
+++ b/Engine/Plugins/EnhancedInput/Source/EnhancedInput/Private/UserSettings/EnhancedInputUserSettings.cpp
@@ -0,0 +1,816 @@
+// Copyright Epic Games, Inc. All Rights Reserved.
+
+#include "UserSettings/EnhancedInputUserSettings.h"
+
+#include "EnhancedPlayerInput.h"
+#include "EnhancedInputModule.h"
+#include "InputMappingContext.h"
+#include "PlayerMappableKeySettings.h"
+#include "EnhancedActionKeyMapping.h"
+#include "EnhancedInputSubsystems.h"
+#include "Engine/LocalPlayer.h"
+#include "GameFramework/PlayerController.h"
+#include "Kismet/GameplayStatics.h"
+#include "GameFramework/InputDeviceSubsystem.h"
+#include "EnhancedInputDeveloperSettings.h"
+#include "Internationalization/Text.h"
+#include "GameFramework/InputSettings.h"
+#include "HAL/IConsoleManager.h"
+#include "EnhancedInputLibrary.h"
+#include "Algo/Find.h"
+#include "NativeGameplayTags.h"
+
+#define LOCTEXT_NAMESPACE "EnhancedInputMappableUserSettings"
+
+namespace UE::EnhancedInput
+{
+	/** The name of the slot that these settings will save to */
+	static const FString SETTINGS_SLOT_NAME = TEXT("EnhancedInputUserSettings");
+
+	UE_DEFINE_GAMEPLAY_TAG_STATIC(TAG_DefaultProfileIdentifier, "InputUserSettings.FailureReasons.InvalidActionName");
+	static const FText DefaultProfileDisplayName = LOCTEXT("Default_Profile_name", "Default Profile");
+	
+	UE_DEFINE_GAMEPLAY_TAG_STATIC(TAG_InvalidActionName, "InputUserSettings.FailureReasons.InvalidActionName");
+	UE_DEFINE_GAMEPLAY_TAG_STATIC(TAG_NoKeyProfile, "InputUserSettings.FailureReasons.NoKeyProfile");
+	UE_DEFINE_GAMEPLAY_TAG_STATIC(TAG_NoMatchingMappings, "InputUserSettings.FailureReasons.NoMatchingMappings");
+	
+	// TODO: Break this out to somewhere else, probably the EI library. 
+	static ULocalPlayer* GetLocalPlayer(const UEnhancedPlayerInput* PlayerInput)
+	{
+		if (PlayerInput)
+		{
+			if (APlayerController* PC = Cast<APlayerController>(PlayerInput->GetOuter()))
+			{
+				return PC->GetLocalPlayer();
+			}
+		}
+		return nullptr;
+	}
+
+	static void DumpAllKeyProfilesToLog(const TArray<FString>& Args)
+	{
+		// Dump every local player subsystem's logs
+		/*UEnhancedInputLibrary::ForEachSubsystem([Args](IEnhancedInputSubsystemInterface* Subsystem)
+		{
+			if (const UEnhancedInputUserSettings* Settings = Subsystem->GetUserSettings())
+			{
+				if (const UEnhancedPlayerMappableKeyProfile* Profile = Settings->GetCurrentKeyProfile())
+				{
+					Profile->DumpProfileToLog();	
+				}
+			}
+		});*/
+	}
+	
+	static FAutoConsoleCommand ConsoelCommandDumpProfileToLog(
+		TEXT("EnhancedInput.DumpKeyProfileToLog"),
+		TEXT(""),
+		FConsoleCommandWithArgsDelegate::CreateStatic(UE::EnhancedInput::DumpAllKeyProfilesToLog));
+}
+
+///////////////////////////////////////////////////////////
+// FMapPlayerKeyArgs
+
+FMapPlayerKeyArgs::FMapPlayerKeyArgs()
+	: ActionName(NAME_None)
+	, Slot(EPlayerMappableKeySlot::Unspecified)
+	, NewKey(EKeys::Invalid)
+	, HardwareDeviceId(NAME_None)
+	, bCreateMatchingSlotIfNeeded(true)
+{
+}
+
+///////////////////////////////////////////////////////////
+// FPlayerKeyMapping
+
+FPlayerKeyMapping::FPlayerKeyMapping()
+	: ActionName(NAME_None)
+	, DisplayName(FText::GetEmpty())
+	, Slot(EPlayerMappableKeySlot::Unspecified)
+	, DefaultKey(EKeys::Invalid)
+	, CurrentKey(EKeys::Invalid)
+	, HardwareDeviceId(FHardwareDeviceIdentifier::Invalid)
+{
+}
+
+// The default constructor creates an invalid mapping. Use this as a way to return references
+// to an invalid mapping for BP functions
+FPlayerKeyMapping FPlayerKeyMapping::InvalidMapping = FPlayerKeyMapping();
+
+bool FPlayerKeyMapping::IsCustomized() const
+{
+	return CurrentKey.IsValid() && (CurrentKey != DefaultKey);
+}
+
+bool FPlayerKeyMapping::IsValid() const
+{
+	return ActionName.IsValid() && CurrentKey.IsValid();
+}
+
+const FKey& FPlayerKeyMapping::GetCurrentKey() const
+{
+	return IsCustomized() ? CurrentKey : DefaultKey;
+}
+
+const FKey& FPlayerKeyMapping::GetDefaultKey() const
+{
+	return DefaultKey;
+}
+
+FString FPlayerKeyMapping::ToString() const
+{
+	const UEnum* PlayerMappableEnumClass = StaticEnum<EPlayerMappableKeySlot>();
+	check(PlayerMappableEnumClass);
+	return
+		FString::Printf(TEXT("Action Name: '%s'  Slot: '%s'  Default Key: '%s'  Player Mapped Key: '%s'  HardwareDevice:  '%s'"),
+			*ActionName.ToString(),
+			*PlayerMappableEnumClass->GetNameStringByValue(static_cast<int64>(Slot)),
+			*DefaultKey.ToString(),
+			*CurrentKey.ToString(),
+			*HardwareDeviceId.ToString());
+}
+
+const FName FPlayerKeyMapping::GetActionName() const
+{
+	return ActionName;
+}
+
+const FText& FPlayerKeyMapping::GetDisplayName() const
+{
+	// Just in case the display name is empty on this mapping, see if we can fall back to the original mapping copy's display name
+	if (DisplayName.IsEmpty() && !OriginalMappingCopy.PlayerMappableOptions.DisplayName.IsEmpty())
+	{
+		return OriginalMappingCopy.PlayerMappableOptions.DisplayName;
+	}
+	return DisplayName;
+}
+
+EPlayerMappableKeySlot FPlayerKeyMapping::GetSlot() const
+{
+	return Slot;
+}
+
+const FHardwareDeviceIdentifier& FPlayerKeyMapping::GetHardwareDeviceId() const
+{
+	return HardwareDeviceId;
+}
+
+uint32 GetTypeHash(const FPlayerKeyMapping& InMapping)
+{
+	uint32 Hash = 0;
+	Hash = HashCombine(Hash, GetTypeHash(InMapping.ActionName));
+	Hash = HashCombine(Hash, GetTypeHash(InMapping.Slot));
+	Hash = HashCombine(Hash, GetTypeHash(InMapping.CurrentKey));
+	Hash = HashCombine(Hash, GetTypeHash(InMapping.HardwareDeviceId));
+	return Hash;
+}
+
+void FPlayerKeyMapping::ResetToDefault()
+{
+	CurrentKey = DefaultKey;
+}
+
+void FPlayerKeyMapping::SetCurrentKey(const FKey& NewKey)
+{
+	CurrentKey = NewKey;
+}
+
+bool FPlayerKeyMapping::operator==(const FPlayerKeyMapping& Other) const
+{
+	return
+	       ActionName			== Other.ActionName
+		&& Slot					== Other.Slot
+		&& HardwareDeviceId		== Other.HardwareDeviceId
+		&& CurrentKey		== Other.CurrentKey
+		&& OriginalMappingCopy	== Other.OriginalMappingCopy;
+}
+
+bool FPlayerKeyMapping::operator!=(const FPlayerKeyMapping& Other) const
+{
+	return !FPlayerKeyMapping::operator==(Other);
+}
+
+///////////////////////////////////////////////////////////
+// UEnhancedPlayerMappableKeyProfile
+
+void UEnhancedPlayerMappableKeyProfile::ResetToDefault()
+{
+	// Reset every player mapping to the default key value
+	for (TPair<FName, FKeyMappingRow>& Pair : PlayerMappedKeys)
+	{
+		for (FPlayerKeyMapping& Mapping : Pair.Value.Mappings)
+		{
+			Mapping.ResetToDefault();
+		}
+	}
+	
+	UE_LOG(LogEnhancedInput, Verbose, TEXT("Reset Player Mappable Key Profile '%s' to default values"), *ProfileIdentifier.ToString());
+}
+
+void UEnhancedPlayerMappableKeyProfile::EquipProfile()
+{
+	UE_LOG(LogEnhancedInput, Verbose, TEXT("Equipping Player Mappable Key Profile '%s'"), *ProfileIdentifier.ToString());
+}
+
+void UEnhancedPlayerMappableKeyProfile::UnEquipProfile()
+{
+	UE_LOG(LogEnhancedInput, Verbose, TEXT("Unequipping Player Mappable Key Profile '%s'"), *ProfileIdentifier.ToString());
+}
+
+void UEnhancedPlayerMappableKeyProfile::SetDisplayName(const FText& NewDisplayName)
+{
+	DisplayName = NewDisplayName;
+}
+
+const FGameplayTag& UEnhancedPlayerMappableKeyProfile::GetProfileIdentifer() const
+{
+	return ProfileIdentifier;
+}
+
+const FText& UEnhancedPlayerMappableKeyProfile::GetProfileDisplayName() const
+{
+	return DisplayName;
+}
+
+const TMap<FName, FKeyMappingRow>& UEnhancedPlayerMappableKeyProfile::GetPlayerMappedActions() const
+{
+	return PlayerMappedKeys;
+}
+
+const FKeyMappingRow* UEnhancedPlayerMappableKeyProfile::FindKeyMappingRow(const FName InActionName) const
+{
+	return PlayerMappedKeys.Find(InActionName);
+}
+
+void UEnhancedPlayerMappableKeyProfile::DumpProfileToLog() const
+{
+	UE_LOG(LogEnhancedInput, Log, TEXT("%s"), *ToString());
+}
+
+FString UEnhancedPlayerMappableKeyProfile::ToString() const
+{
+	TStringBuilder<1024> Builder;
+	Builder.Appendf(TEXT("Key Profile '%s' has %d key mappings\n"), *ProfileIdentifier.ToString(), PlayerMappedKeys.Num());
+	
+	for (const TPair<FName, FKeyMappingRow>& Pair : PlayerMappedKeys)
+	{
+		Builder.Append(Pair.Key.ToString());
+		
+		for (const FPlayerKeyMapping& Mapping : Pair.Value.Mappings)
+		{
+			Builder.Append(Mapping.ToString());
+			Builder.Append("\n\t");
+		}
+	}
+
+	return Builder.ToString();
+}
+
+void UEnhancedPlayerMappableKeyProfile::Serialize(FArchive& Ar)
+{
+	// See note in header!
+	Super::Serialize(Ar);
+}
+
+FPlayerKeyMapping* UEnhancedPlayerMappableKeyProfile::FindKeyMapping(const FMapPlayerKeyArgs& InArgs) const
+{
+	// Get the current mappings for the desired action name.
+	if (const FKeyMappingRow* MappingRow = PlayerMappedKeys.Find(InArgs.ActionName))
+	{
+		// If mapping already exists for the given slot and hardware device, then we can
+		// just change that key
+		const FPlayerKeyMapping* Mapping = Algo::FindByPredicate(MappingRow->Mappings, [&InArgs](const FPlayerKeyMapping& Mapping)
+		{
+			return Mapping.GetSlot() == InArgs.Slot && Mapping.GetHardwareDeviceId().HardwareDeviceIdentifier == InArgs.HardwareDeviceId;
+		});
+		return const_cast<FPlayerKeyMapping*>(Mapping);
+	}
+	return nullptr;
+}
+
+///////////////////////////////////////////////////////////
+// UEnhancedInputUserSettings
+
+UEnhancedInputUserSettings* UEnhancedInputUserSettings::LoadOrCreateSettings(UEnhancedPlayerInput* PlayerInput)
+{
+	UEnhancedInputUserSettings* Settings = nullptr;
+
+	ULocalPlayer* LocalPlayer = UE::EnhancedInput::GetLocalPlayer(PlayerInput);
+
+	if (!LocalPlayer)
+	{
+		UE_LOG(LogEnhancedInput, Error, TEXT("Unable to determine an owning Local Player for the given Enhanced Player Input object"));
+		return nullptr;
+	}
+	
+	// If the save game exists, load it.
+	if (UGameplayStatics::DoesSaveGameExist(UE::EnhancedInput::SETTINGS_SLOT_NAME, LocalPlayer->GetLocalPlayerIndex()))
+	{
+		USaveGame* Slot = UGameplayStatics::LoadGameFromSlot(UE::EnhancedInput::SETTINGS_SLOT_NAME, LocalPlayer->GetLocalPlayerIndex());
+		Settings = Cast<UEnhancedInputUserSettings>(Slot);
+	}
+	
+	if (Settings == nullptr)
+	{
+		Settings = Cast<UEnhancedInputUserSettings>(UGameplayStatics::CreateSaveGameObject(UEnhancedInputUserSettings::StaticClass()));
+	}
+
+	Settings->Initialize(PlayerInput);
+	Settings->ApplySettings();
+
+	return Settings;
+}
+
+void UEnhancedInputUserSettings::Initialize(UEnhancedPlayerInput* InPlayerInput)
+{
+	OwningPlayerInput = InPlayerInput;
+	ensureMsgf(GetPlayerInput(), TEXT("UEnhancedInputUserSettings is missing a player input!"));
+
+	// Create a default key mapping profile in the case where one doesn't exist
+	if (!GetCurrentKeyProfile())
+	{
+		FPlayerMappableKeyProfileCreationArgs Args = {};
+		Args.ProfileIdentifier = UE::EnhancedInput::TAG_DefaultProfileIdentifier;
+		Args.DisplayName = UE::EnhancedInput::DefaultProfileDisplayName;
+		Args.bSetAsCurrentProfile = true;
+		
+		CreateNewKeyProfile(Args);
+	}
+}
+
+void UEnhancedInputUserSettings::ApplySettings()
+{
+	ensureMsgf(GetPlayerInput(), TEXT("UEnhancedInputUserSettings is missing a player input!"));
+	UE_LOG(LogEnhancedInput, Verbose, TEXT("Enhanced Input User Settings applied!"));
+}
+
+void UEnhancedInputUserSettings::SaveSettings()
+{
+	ensureMsgf(GetPlayerInput(), TEXT("UEnhancedInputUserSettings is missing a player input!"));
+	
+	if (ULocalPlayer* OwningPlayer = GetLocalPlayer())
+	{
+		UGameplayStatics::SaveGameToSlot(this, UE::EnhancedInput::SETTINGS_SLOT_NAME, OwningPlayer->GetLocalPlayerIndex());
+		UE_LOG(LogEnhancedInput, Verbose, TEXT("Enhanced Input User Settings saved!"));
+	}
+	else
+	{
+		UE_LOG(LogEnhancedInput, Warning, TEXT("Attempting to save Enhanced Input User settings without an owning local player!"));
+	}
+}
+
+namespace UE::EnhancedInput
+{
+	static const int32 GPlayerMappableSaveVersion = 1;
+
+	/** Struct used to store info about the mappable profile subobjects */
+	struct FMappableKeysHeader
+	{
+		friend FArchive& operator<<(FArchive& Ar, FMappableKeysHeader& Header)
+		{
+			Ar << Header.ProfileIdentifier;
+			Ar << Header.ClassPath;
+			Ar << Header.ObjectPath;
+			return Ar;
+		}
+
+		FGameplayTag ProfileIdentifier;
+		FString ClassPath;
+		FString ObjectPath;
+	};
+}
+
+void UEnhancedInputUserSettings::Serialize(FArchive& Ar)
+{
+	Super::Serialize(Ar);
+
+	if (IsTemplate() || Ar.IsCountingMemory())
+	{
+		return;
+	}
+	
+	int32 SaveVersion = UE::EnhancedInput::GPlayerMappableSaveVersion;
+	Ar << SaveVersion;
+	
+	// Detect a mis-match of byte streams, i.e. file corruption
+	ensure(SaveVersion == UE::EnhancedInput::GPlayerMappableSaveVersion);
+	
+	TArray<UE::EnhancedInput::FMappableKeysHeader> Headers;
+	if (Ar.IsSaving())
+	{
+		UObject* Outer = this;
+		
+		for (TPair<FGameplayTag, UEnhancedPlayerMappableKeyProfile*> ProfilePair : SavedKeyProfiles)
+		{
+			Headers.Push({ ProfilePair.Key, ProfilePair.Value->GetClass()->GetPathName(), ProfilePair.Value->GetPathName(Outer) });
+		}
+	}
+	
+	Ar << Headers;
+
+	if (Ar.IsLoading())
+	{
+		for (UE::EnhancedInput::FMappableKeysHeader& Header : Headers)
+		{
+			if (const UClass* FoundClass = FindObject<UClass>(nullptr, *Header.ClassPath))
+			{
+				UEnhancedPlayerMappableKeyProfile* NewProfile = NewObject<UEnhancedPlayerMappableKeyProfile>(/* outer */ this, /* class */ FoundClass);
+				SavedKeyProfiles.Add(Header.ProfileIdentifier, NewProfile);
+			}
+		}
+	}
+
+	FString SavedObjectTerminator = TEXT("ObjectEnd");
+	
+	for (TPair<FGameplayTag, UEnhancedPlayerMappableKeyProfile*> ProfilePair : SavedKeyProfiles)
+	{
+		ProfilePair.Value->Serialize(Ar);
+
+		// Save a terminator after each subobject
+		Ar << SavedObjectTerminator;
+		
+		if (!ensure(SavedObjectTerminator == TEXT("ObjectEnd")))
+		{
+			UE_LOG(LogEnhancedInput, Error, TEXT("Serialization size mismatch! Possible over-read or over-write of this buffer."));
+			break;
+		}
+	}
+}
+
+UEnhancedPlayerInput* UEnhancedInputUserSettings::GetPlayerInput() const
+{
+	return OwningPlayerInput;
+}
+
+APlayerController* UEnhancedInputUserSettings::GetPlayerController() const
+{
+	if (UEnhancedPlayerInput* PlayerInput = GetPlayerInput())
+	{
+		return Cast<APlayerController>(PlayerInput->GetOuter());
+	}
+	return nullptr;
+}
+
+ULocalPlayer* UEnhancedInputUserSettings::GetLocalPlayer() const
+{
+	if (APlayerController* PC = GetPlayerController())
+	{
+		return PC->GetLocalPlayer();
+	}
+	return nullptr;
+}
+
+void UEnhancedInputUserSettings::MapPlayerKey(const FMapPlayerKeyArgs& InArgs, FGameplayTagContainer& FailureReason)
+{
+	if (!InArgs.ActionName.IsValid())
+	{
+		FailureReason.AddTag(UE::EnhancedInput::TAG_InvalidActionName);
+		return;
+	}
+	
+	// Get the key profile that was specific
+	UEnhancedPlayerMappableKeyProfile* KeyProfile = InArgs.ProfileId.IsValid() ? GetKeyProfileWithIdentifier(InArgs.ProfileId) : GetCurrentKeyProfile();
+	if (!KeyProfile)
+	{
+		FailureReason.AddTag(UE::EnhancedInput::TAG_NoKeyProfile);
+		return;
+	}
+
+	// If this mapping already exists, we can simply change it's key and be done
+	if (FPlayerKeyMapping* FoundMapping = KeyProfile->FindKeyMapping(InArgs))
+	{
+		// Then set the player mapped key
+		FoundMapping->SetCurrentKey(InArgs.NewKey);
+		OnSettingsChanged.Broadcast(this);
+		return;
+	}
+	// If it doesn't exist, then we need to make it if there is a valid action name
+	else if (FKeyMappingRow* MappingRow = KeyProfile->PlayerMappedKeys.Find(InArgs.ActionName))
+	{
+		// If one doesn't exist, then we need to create a new mapping in the given
+		// slot. 
+		// In order to populate the default values correctly, we only do this if we know that
+		// mappings exist for it
+		const int32 NumMappings = MappingRow->Mappings.Num();
+		if (InArgs.bCreateMatchingSlotIfNeeded && NumMappings > 0)
+		{
+			const auto ExistingMapping = MappingRow->Mappings.begin();
+			
+			// Add a default mapping to this row
+			FPlayerKeyMapping PlayerMappingData = {};
+			PlayerMappingData.ActionName = InArgs.ActionName;
+			PlayerMappingData.Slot = InArgs.Slot;
+
+			// If there is some valid hardware then keep track of that
+			if (const UInputPlatformSettings* PlatformSettings = UInputPlatformSettings::Get())
+			{
+				if (const FHardwareDeviceIdentifier* Hardware = PlatformSettings->GetHardwareDeviceForClassName(InArgs.HardwareDeviceId))
+				{
+					PlayerMappingData.HardwareDeviceId = *Hardware;	
+				}
+				else
+				{
+					UE_LOG(LogEnhancedInput, Log, TEXT("[UEnhancedInputUserSettings::MapPlayerKey] Unable to find a matching Hardware Device Identifier with the HardwareDeviceId of '%s'"), *InArgs.HardwareDeviceId.ToString());
+				}
+			}
+			
+			// This mapping never existed in the default IMC, so the default mapping will be the default
+			// EKeys::Invalid and we only need to track the player mapped key
+			PlayerMappingData.SetCurrentKey(InArgs.NewKey);
+			PlayerMappingData.DisplayName = ExistingMapping->DisplayName;
+			MappingRow->Mappings.Add(PlayerMappingData);
+			OnSettingsChanged.Broadcast(this);
+		}
+	}
+}
+
+void UEnhancedInputUserSettings::UnMapPlayerKey(const FMapPlayerKeyArgs& InArgs, FGameplayTagContainer& FailureReason)
+{
+	if (!InArgs.ActionName.IsValid())
+	{
+		FailureReason.AddTag(UE::EnhancedInput::TAG_InvalidActionName);
+		return;
+	}
+	
+	// Get the key profile that was specified
+	UEnhancedPlayerMappableKeyProfile* KeyProfile = InArgs.ProfileId.IsValid() ? GetKeyProfileWithIdentifier(InArgs.ProfileId) : GetCurrentKeyProfile();
+	if (!KeyProfile)
+	{
+		FailureReason.AddTag(UE::EnhancedInput::TAG_NoKeyProfile);
+		return;
+	}
+
+	if (FPlayerKeyMapping* FoundMapping = KeyProfile->FindKeyMapping(InArgs))
+	{
+		// Then set the player mapped key
+		FoundMapping->ResetToDefault();
+		OnSettingsChanged.Broadcast(this);
+		
+		UE_LOG(LogEnhancedInput, Verbose, TEXT("[UEnhancedInputUserSettings::MapPlayerKey] Reset keymapping to default: '%s'"), *FoundMapping->ToString());
+		
+		return;
+	}
+	// if a mapping doesn't exist, then we can't unmap it
+	else
+	{
+		FailureReason.AddTag(UE::EnhancedInput::TAG_NoMatchingMappings);
+	}
+}
+
+const TSet<FPlayerKeyMapping>& UEnhancedInputUserSettings::FindMappingsForAction(const FName ActionName) const
+{
+	if (UEnhancedPlayerMappableKeyProfile* KeyProfile = GetCurrentKeyProfile())
+	{
+		FKeyMappingRow& ExistingMappings = KeyProfile->PlayerMappedKeys.FindOrAdd(ActionName);
+		return ExistingMappings.Mappings;
+	}
+	
+	UE_LOG(LogEnhancedInput, Error, TEXT("There is no current mappable key profile! No mappings will be returned."));
+	
+	static TSet<FPlayerKeyMapping> EmptyMappings;
+	return EmptyMappings;
+}
+
+const FPlayerKeyMapping* UEnhancedInputUserSettings::FindCurrentMappingForSlot(const FName ActionName, const EPlayerMappableKeySlot InSlot) const
+{
+	const TSet<FPlayerKeyMapping>& AllMappings = FindMappingsForAction(ActionName);
+	for (const FPlayerKeyMapping& Mapping : AllMappings)
+	{
+		if (Mapping.Slot == InSlot)
+		{
+			return &Mapping;
+		}
+	}
+	
+	UE_LOG(LogEnhancedInput, Warning, TEXT("No mappings could be found for action '%s'"), *ActionName.ToString());
+	return nullptr;
+}
+
+bool UEnhancedInputUserSettings::SetKeyProfile(const FGameplayTag& InProfileId)
+{
+	UEnhancedPlayerInput* PlayerInput = GetPlayerInput();
+	if (!PlayerInput)
+	{
+		UE_LOG(LogEnhancedInput, Error, TEXT("Failed to find the PlayerInput associated with the Enhanced Input user settings!"), *InProfileId.ToString());
+		return false;
+	}
+
+	const FGameplayTag OriginalProfileId = CurrentProfileIdentifier;
+	UEnhancedPlayerMappableKeyProfile* OriginalProfile = GetCurrentKeyProfile();
+	
+	if (const TObjectPtr<UEnhancedPlayerMappableKeyProfile>* NewProfile = SavedKeyProfiles.Find(InProfileId))
+	{
+		// Unequip the original profile if there was one
+		if (OriginalProfile)
+		{
+			OriginalProfile->UnEquipProfile();
+		}
+
+		// Equip the new profile
+		NewProfile->Get()->EquipProfile();
+
+		// Keep track of what the current profile is now
+		CurrentProfileIdentifier = InProfileId;
+
+		// Let any listeners know that the mapping profile has changed
+		OnKeyProfileChanged.Broadcast(NewProfile->Get());
+	}
+	else
+	{
+		UE_LOG(LogEnhancedInput, Error, TEXT("No profile with name '%s' exists! Did you call CreateNewKeyProfile at any point?"), *InProfileId.ToString());
+		return false;
+	}
+
+	UE_LOG(LogEnhancedInput, Verbose, TEXT("Successfully changed Key Profile from '%s' to '%s'"), *OriginalProfileId.ToString(), *CurrentProfileIdentifier.ToString());
+	return true;
+}
+
+const FGameplayTag& UEnhancedInputUserSettings::GetCurrentKeyProfileIdentifier() const
+{
+	return CurrentProfileIdentifier;
+}
+
+UEnhancedPlayerMappableKeyProfile* UEnhancedInputUserSettings::GetCurrentKeyProfile() const
+{
+	return GetKeyProfileWithIdentifier(CurrentProfileIdentifier);
+}
+
+FPlayerMappableKeyProfileCreationArgs::FPlayerMappableKeyProfileCreationArgs()
+	: ProfileType(GetDefault<UEnhancedInputDeveloperSettings>()->DefaultPlayerMappableKeyProfileClass.Get())
+	, ProfileIdentifier(FGameplayTag::EmptyTag)
+	, DisplayName(FText::GetEmpty())
+	, bSetAsCurrentProfile(true)
+{
+}
+
+UEnhancedPlayerMappableKeyProfile* UEnhancedInputUserSettings::CreateNewKeyProfile(const FPlayerMappableKeyProfileCreationArgs& InArgs)
+{
+	if (!InArgs.ProfileType)
+	{
+		UE_LOG(LogEnhancedInput, Error, TEXT("Invalid ProfileType given CreateNewKeyProfile!"));
+		return nullptr;
+	}
+	
+	UEnhancedPlayerMappableKeyProfile* OutProfile = GetKeyProfileWithIdentifier(InArgs.ProfileIdentifier);
+	
+	// Check for an existing profile of this name
+	if (OutProfile)
+	{
+		UE_LOG(LogEnhancedInput, Warning, TEXT("A key profile with the name '%s' already exists! Use a different name."), *InArgs.ProfileIdentifier.ToString());
+	}
+	else
+	{
+		// Create a new mapping profile
+		OutProfile = NewObject<UEnhancedPlayerMappableKeyProfile>(/* outer */ this, /* class */ InArgs.ProfileType);
+		OutProfile->ProfileIdentifier = InArgs.ProfileIdentifier;
+		OutProfile->DisplayName = InArgs.DisplayName;
+		
+		SavedKeyProfiles.Add(InArgs.ProfileIdentifier, OutProfile);
+	}
+	
+	// set as current
+	if (InArgs.bSetAsCurrentProfile)
+	{
+		SetKeyProfile(InArgs.ProfileIdentifier);
+	}
+
+	UE_LOG(LogEnhancedInput, Verbose, TEXT("Completed creation of key mapping profile '%s'"), *OutProfile->ProfileIdentifier.ToString());
+	
+	return OutProfile;
+}
+
+UEnhancedPlayerMappableKeyProfile* UEnhancedInputUserSettings::GetKeyProfileWithIdentifier(const FGameplayTag& ProfileId) const
+{
+	if (const TObjectPtr<UEnhancedPlayerMappableKeyProfile>* ExistingProfile = SavedKeyProfiles.Find(ProfileId))
+	{		
+		return ExistingProfile->Get();
+	}
+	return nullptr;
+}
+
+bool UEnhancedInputUserSettings::RegisterInputMappingContexts(const TSet<UInputMappingContext*>& MappingContexts)
+{
+	bool bResult = false;
+
+	for (UInputMappingContext* IMC : MappingContexts)
+	{
+		bResult |= RegisterInputMappingContext(IMC);
+	}
+
+	return bResult;
+}
+
+bool UEnhancedInputUserSettings::RegisterInputMappingContext(UInputMappingContext* IMC)
+{
+	if (!IMC)
+	{
+		UE_LOG(LogEnhancedInput, Error, TEXT("Attempting to register a null mapping context with the user settings!"));
+		ensure(false);
+		return false;
+	}
+
+	// There is no need to register an IMC if it already is
+	if( RegisteredMappingContexts.Contains(IMC))
+	{
+		return false;
+	}
+	
+	const FSetElementId Res = RegisteredMappingContexts.Add(IMC);
+
+	UEnhancedPlayerMappableKeyProfile* CurrentProfile = GetCurrentKeyProfile();
+	if (!CurrentProfile)
+	{
+		UE_LOG(LogEnhancedInput, Error, TEXT("There is not an active key profile!"));
+		ensure(false);
+		return false;
+	}
+	
+	for (const FEnhancedActionKeyMapping& KeyMapping : IMC->GetMappings())
+	{
+		// Skip over non-player mappable keys
+		if (!KeyMapping.IsPlayerMappable())
+		{
+			continue;
+		}
+
+		// Get the unique FName for the "Action name". This is set on the UInputAction
+		// but can be overriden on each FEnhancedActionKeyMapping to create multiple
+		// mapping options for a single Input Action. 
+		const FName ActionName = KeyMapping.GetMappingName();
+
+		// If a mapping row exists for this action name, then check if
+		bool bMappingIsInitalized = false;
+		if (FKeyMappingRow* ExistingMappings = CurrentProfile->PlayerMappedKeys.Find(ActionName))
+		{
+			// Check if we actually have to initalize a new mapping or not
+			for (FPlayerKeyMapping& ExistingMapping : ExistingMappings->Mappings)
+			{
+				// If the original mapping ha already been set, then we know that it's initalized already.
+				// If we don't do this, then registering an IMC will overrwrite whatever saved settings the user has made
+				if (ExistingMapping.OriginalMappingCopy == KeyMapping)
+				{
+					bMappingIsInitalized = true;
+					break;
+				}
+			}
+		}
+		
+		if (!bMappingIsInitalized)
+		{
+			FKeyMappingRow& MappingRow = CurrentProfile->PlayerMappedKeys.Add(ActionName);		
+
+			// Add a default mapping to this row
+			FPlayerKeyMapping PlayerMappingData = {};
+			PlayerMappingData.ActionName = ActionName;
+			PlayerMappingData.DefaultKey = KeyMapping.Key;
+			PlayerMappingData.DisplayName = KeyMapping.PlayerMappableOptions.DisplayName;
+			PlayerMappingData.OriginalMappingCopy = KeyMapping;
+
+			// By default, the slot will be determined by how many mappings this action has already.
+			// So if this is the first default mapping, then this will be EPlayerMappableKeySlot::First,
+			// if this is the second element going into this set then it will be EPlayerMappableKeySlot::Second
+			// and so on
+			const uint8 DesiredSlot = FMath::Min<uint8>(static_cast<uint8>(MappingRow.Mappings.Num()), static_cast<uint8>(EPlayerMappableKeySlot::Max));
+			PlayerMappingData.Slot = (EPlayerMappableKeySlot) DesiredSlot;
+		
+			MappingRow.Mappings.Add(PlayerMappingData);
+		}
+	}
+
+	if (Res.IsValidId())
+	{
+		OnMappingContextRegistered.Broadcast(IMC);	
+	}
+
+	UE_LOG(LogEnhancedInput, Verbose, TEXT("Registered IMC with UEnhancedInputUserSettings: %s"), *IMC->GetFName().ToString());
+	return Res.IsValidId();
+}
+
+bool UEnhancedInputUserSettings::UnregisterInputMappingContext(const UInputMappingContext* IMC)
+{
+	return RegisteredMappingContexts.Remove(IMC) != INDEX_NONE;
+}
+
+bool UEnhancedInputUserSettings::UnregisterInputMappingContexts(const TSet<UInputMappingContext*>& MappingContexts)
+{
+	bool bResult = false;
+
+	for (UInputMappingContext* IMC : MappingContexts)
+	{
+		bResult |= UnregisterInputMappingContext(IMC);
+	}
+
+	return bResult;
+}
+
+const TSet<TObjectPtr<UInputMappingContext>>& UEnhancedInputUserSettings::GetRegisteredInputMappingContexts() const
+{
+	return RegisteredMappingContexts;
+}
+
+bool UEnhancedInputUserSettings::IsMappingContextRegistered(const UInputMappingContext* IMC) const
+{
+	return RegisteredMappingContexts.Contains(IMC);
+}
+
+#undef LOCTEXT_NAMESPACE
\ No newline at end of file
diff --git a/Engine/Plugins/EnhancedInput/Source/EnhancedInput/Public/EnhancedInputDeveloperSettings.h b/Engine/Plugins/EnhancedInput/Source/EnhancedInput/Public/EnhancedInputDeveloperSettings.h
--- a/Engine/Plugins/EnhancedInput/Source/EnhancedInput/Public/EnhancedInputDeveloperSettings.h
+++ b/Engine/Plugins/EnhancedInput/Source/EnhancedInput/Public/EnhancedInputDeveloperSettings.h
@@ -7,6 +7,10 @@
 #include "EnhancedInputDeveloperSettings.generated.h"
 
 class UInputMappingContext;
+class UEnhancedInputUserSettings;
+class UEnhancedPlayerMappableKeyProfile;
+class UEnhancedPlayerInput;
+enum class EPlayerMappableKeySlot : uint8;
 
 /** Represents a single input mapping context and the priority that it should be applied with */
 USTRUCT()
@@ -23,8 +27,6 @@ struct FDefaultContextSetting
 	int32 Priority = 0;
 };
 
-class UEnhancedPlayerInput;
-
 /** Developer settings for Enhanced Input */
 UCLASS(config = Input, defaultconfig, meta = (DisplayName = "Enhanced Input"))
 class ENHANCEDINPUT_API UEnhancedInputDeveloperSettings : public UDeveloperSettingsBackedByCVars
@@ -47,10 +49,6 @@ public:
 	 */
 	UPROPERTY(config, EditAnywhere, Category = "Enhanced Input|World Subsystem", meta = (editCondition = "bEnableDefaultMappingContexts && bEnableWorldSubsystem"))
 	TArray<FDefaultContextSetting> DefaultWorldSubsystemMappingContexts;
-	
-	/** The default player input class that the Enhanced Input world subsystem will use. */
-	UPROPERTY(config, EditAnywhere, NoClear, Category = "Enhanced Input|World Subsystem", meta=(editCondition = "bEnableWorldSubsystem"))
-	TSoftClassPtr<UEnhancedPlayerInput> DefaultWorldInputClass;
 
 	/**
 	 * Platform specific settings for Enhanced Input.
@@ -59,6 +57,37 @@ public:
 	UPROPERTY(EditAnywhere, Category = "Enhanced Input")
 	FPerPlatformSettings PlatformSettings;
 
+	/**
+	 * The class that should be used for the User Settings by each Enhanced Input subsystem.
+	 * An instance of this class will be spawned by each Enhanced Input subsytem as a place to store
+	 * user settings such as keymappings, accessibility settings, etc. Subclass this to add more custom
+	 * options to your game.
+	 *
+	 * Note: This is a new experimental feature!
+	 */
+	UPROPERTY(config, EditAnywhere, NoClear, Category = "Enhanced Input|User Settings", meta=(editCondition = "bEnableUserSettings"))
+	TSoftClassPtr<UEnhancedInputUserSettings> UserSettingsClass;
+
+	/**
+	 * The default class for the player mappable key profile, used to store the key mappings set by the player in the user settings.
+	 * 
+	 * Note: This is a new experimental feature!
+	 */
+	UPROPERTY(config, EditAnywhere, NoClear, Category = "Enhanced Input|User Settings", meta=(editCondition = "bEnableUserSettings"))
+	TSoftClassPtr<UEnhancedPlayerMappableKeyProfile> DefaultPlayerMappableKeyProfileClass;
+	
+	/** The default player input class that the Enhanced Input world subsystem will use. */
+	UPROPERTY(config, EditAnywhere, NoClear, Category = "Enhanced Input|World Subsystem", meta=(editCondition = "bEnableWorldSubsystem"))
+	TSoftClassPtr<UEnhancedPlayerInput> DefaultWorldInputClass;
+
+	/**
+	 * If true, then an instance of the User Settings Class will be created on each Enhanced Input subsytem.
+	 * 
+	 * Note: This is a new experimental feature!
+	 */
+	UPROPERTY(config, EditAnywhere, Category = "Enhanced Input|User Settings", meta=(DisplayPriority = 1))
+	uint8 bEnableUserSettings : 1;
+
 	/** If true, then the DefaultMappingContexts will be applied to all Enhanced Input Subsystems. */
 	UPROPERTY(EditAnywhere, Category = "Enhanced Input", meta = (ConsoleVariable = "EnhancedInput.EnableDefaultMappingContexts"))
 	uint8 bEnableDefaultMappingContexts : 1;
@@ -73,7 +102,7 @@ public:
 	uint8 bShouldOnlyTriggerLastActionInChord : 1;
 	
 	/** If true, then the world subsystem will be created. */
-	UPROPERTY(config, EditAnywhere, Category = "Enhanced Input|World Subsystem", meta=(DisplayName="Enable World Subsystem (Experimental)"))
+	UPROPERTY(config, EditAnywhere, Category = "Enhanced Input|World Subsystem", meta=(DisplayName="Enable World Subsystem (Experimental)", DisplayPriority = 1))
 	uint8 bEnableWorldSubsystem : 1;
 	
 	/**
diff --git a/Engine/Plugins/EnhancedInput/Source/EnhancedInput/Public/UserSettings/EnhancedInputUserSettings.h b/Engine/Plugins/EnhancedInput/Source/EnhancedInput/Public/UserSettings/EnhancedInputUserSettings.h
new file mode 100644
--- /dev/null
+++ b/Engine/Plugins/EnhancedInput/Source/EnhancedInput/Public/UserSettings/EnhancedInputUserSettings.h
@@ -0,0 +1,456 @@
+// Copyright Epic Games, Inc. All Rights Reserved.
+
+#pragma once
+
+#include "GameFramework/SaveGame.h"
+#include "UObject/Object.h"
+#include "InputCoreTypes.h"					// For FKey
+#include "GameFramework/InputSettings.h"	// For FHardwareDeviceIdentifier
+#include "EnhancedActionKeyMapping.h"
+#include "GameplayTagContainer.h"
+
+#include "EnhancedInputUserSettings.generated.h"
+
+class UInputMappingContext;
+class UEnhancedPlayerInput;
+class ULocalPlayer;
+
+/**
+ * The "Slot" that a player mappable key is in.
+ * Used by UI to allow for multiple keys to be bound by the player for a single action
+ * 
+ * | <Action Name>  | Slot 1 | Slot 2 | Slot 3 | Slot.... N |
+ */
+UENUM(BlueprintType)
+enum class EPlayerMappableKeySlot : uint8 
+{
+	// The first key slot
+	First = 0,
+
+	// The second mappable key slot. This is the default max in the project settings
+	Second,
+	
+	Third,
+	Fourth,
+	Fifth,
+	Sixth,
+	Seventh,
+	
+	// A key that isn't in any slot
+	Unspecified,
+	Max
+};
+
+/** Arguments that can be used when mapping a player key */
+USTRUCT(BlueprintType)
+struct ENHANCEDINPUT_API FMapPlayerKeyArgs
+{
+	GENERATED_BODY()
+
+	FMapPlayerKeyArgs();	
+	
+	/**
+	 * The name of the action for this key. This is either the default mapping name from an Input Action asset, or one
+	 * that is overriden in the Input Mapping Context.
+	 */
+	UPROPERTY(EditAnywhere, BlueprintReadWrite, Category="Enhanced Input|User Settings")
+	FName ActionName;
+
+	/** What slot this key mapping is for */
+	UPROPERTY(EditAnywhere, BlueprintReadWrite, Category="Enhanced Input|User Settings")
+	EPlayerMappableKeySlot Slot;
+
+	/** The new Key that this action should be mapped to */
+	UPROPERTY(EditAnywhere, BlueprintReadWrite, Category="Enhanced Input|User Settings")
+	FKey NewKey;
+	
+	/** An OPTIONAL specifier about what kind of hardware this mapping is for. */
+	UPROPERTY(EditAnywhere, BlueprintReadWrite, Category="Enhanced Input|User Settings", meta = (EditCondition = "bFilterByHardwareDeviceId == true", AllowPrivateAccess = true, GetOptions = "Engine.InputPlatformSettings.GetAllHardwareDeviceNames"))
+	FName HardwareDeviceId;
+
+	/** The Key Mapping Profile identifier that this mapping should be set on. If this is empty, then the currently equipped profile will be used. */
+	UPROPERTY(EditAnywhere, BlueprintReadWrite, Category="Enhanced Input|User Settings")
+	FGameplayTag ProfileId;
+
+	/** If there is not a player mapping already with the same Slot and Hardware Device ID, then create a new mapping for this slot. */
+	UPROPERTY(EditAnywhere, BlueprintReadWrite, Category="Enhanced Input|User Settings")
+	uint8 bCreateMatchingSlotIfNeeded : 1;
+};
+
+/** Represents a single key mapping that is set by the player */
+USTRUCT(BlueprintType)
+struct ENHANCEDINPUT_API FPlayerKeyMapping
+{
+	GENERATED_BODY()
+	friend class UEnhancedInputUserSettings;
+	
+public:
+	
+	FPlayerKeyMapping();
+
+	/** A static invalid player key mapping to be used for easy comparisons in blueprint */
+	static FPlayerKeyMapping InvalidMapping;
+
+	/** Returns true if this mapping has been customized by the player, and false if it has not been. */
+	bool IsCustomized() const;
+
+	/** Returns true if this player mapping is valid */
+	bool IsValid() const;
+
+	/**
+	 * Returns the key that the player has mapped. If the player has not mapped one yet, then this returns the
+	 * default key mapping from the input mapping context.
+	 */
+	const FKey& GetCurrentKey() const;
+
+	/** Returns the default key that this mapping is to */
+	const FKey& GetDefaultKey() const;
+
+	/** Print out some debug information about this player key mappings */
+	FString ToString() const;
+
+	/**
+	 * The unique FName associated with this action. This is defined by this mappings owning Input Action
+	 * or the individual Enhanced Action Key Mapping if it is overriden
+	 */
+	const FName GetActionName() const;
+
+	/** The localized display name to use for this mapping */
+	const FText& GetDisplayName() const;
+
+	/** Returns what player mappable slot this mapping is in */
+	EPlayerMappableKeySlot GetSlot() const;
+
+	/** Returns the optional hardware device ID that this mapping is specific to */
+	const FHardwareDeviceIdentifier& GetHardwareDeviceId() const; 
+
+	/** Resets the current mapping to the default one */
+	void ResetToDefault();
+
+	/** Sets the value of the current key to the one given */
+	void SetCurrentKey(const FKey& NewKey);
+	
+	ENHANCEDINPUT_API friend uint32 GetTypeHash(const FPlayerKeyMapping& InMapping);
+	bool operator==(const FPlayerKeyMapping& Other) const;
+	bool operator!=(const FPlayerKeyMapping& Other) const;
+
+protected:
+	
+	/** The name of the action for this key */
+	UPROPERTY(VisibleAnywhere, SaveGame, BlueprintReadOnly, Category="Enhanced Input|User Settings")
+	FName ActionName;
+
+	/** Localized display name of this action */
+	UPROPERTY(VisibleAnywhere, SaveGame, BlueprintReadOnly, Category="Enhanced Input|User Settings")
+	FText DisplayName;
+
+	/** What slot this key is mapped to */
+	UPROPERTY(VisibleAnywhere, SaveGame, BlueprintReadOnly, Category="Enhanced Input|User Settings")
+	EPlayerMappableKeySlot Slot;
+	
+	/** The default key that this mapping was set to in its input mapping context */
+	UPROPERTY(VisibleAnywhere, SaveGame, BlueprintReadOnly, Category="Enhanced Input|User Settings", meta = (AllowPrivateAccess = "True"))
+	FKey DefaultKey;
+	
+	/** The key that the player has mapped this action to */
+	UPROPERTY(VisibleAnywhere, SaveGame,  BlueprintReadOnly, Category="Enhanced Input|User Settings", meta = (AllowPrivateAccess = "True"))
+	FKey CurrentKey;
+
+	/** An optional Hardware Device specifier for this mapping */
+	UPROPERTY(VisibleAnywhere, SaveGame, BlueprintReadOnly, Category="Enhanced Input|User Settings")
+	FHardwareDeviceIdentifier HardwareDeviceId;
+
+	/** A copy of the action key mapping that this was generated from. This will only be set on mappings that were created by a Registered IMC */
+	UPROPERTY(VisibleAnywhere, SaveGame, BlueprintReadOnly, Category="Enhanced Input|User Settings")
+	FEnhancedActionKeyMapping OriginalMappingCopy;
+};
+
+/**
+ * Stores all mappings bound to a single action.
+ *
+ * Since a single action can have multiple bindings to it and this system should be Blueprint friendly,
+ * this needs to be a struct (blueprint don't support nested containers).
+ */
+USTRUCT(BlueprintType)
+struct ENHANCEDINPUT_API FKeyMappingRow
+{
+	GENERATED_BODY()
+
+	UPROPERTY(VisibleAnywhere, SaveGame, BlueprintReadOnly, Category="Enhanced Input|User Settings")
+	TSet<FPlayerKeyMapping> Mappings;
+};
+
+
+/** Represents one "Profile" that a user can have for their player mappable keys */
+UCLASS(BlueprintType)
+class ENHANCEDINPUT_API UEnhancedPlayerMappableKeyProfile : public UObject
+{
+	GENERATED_BODY()
+
+	friend class UEnhancedInputUserSettings;
+
+public:
+
+	//~ Begin UObject Interface
+    /**
+     * Because the key mapping profile is serialized as a subobject of the UEnhancedInputUserSettings and requires
+     * some custom serialization logic, you should not override the Serialize method on your custom key profile.
+     * If you need to add custom serialization logic then you can create a struct UPROPERTY and override the struct's
+     * serialization logic, which will prevent you from running into possible issues with subobjects.
+     */
+    virtual void Serialize(FArchive& Ar) override final;
+    //~ End UObject Interface
+
+	UFUNCTION(BlueprintCallable, Category="Enhanced Input|User Settings")
+	virtual void SetDisplayName(const FText& NewDisplayName);
+	
+	UFUNCTION(BlueprintCallable, BlueprintPure, Category="Enhanced Input|User Settings")
+	const FGameplayTag& GetProfileIdentifer() const;
+
+	/** Get the localized display name for this profile */
+	UFUNCTION(BlueprintCallable, BlueprintPure, Category="Enhanced Input|User Settings")
+	const FText& GetProfileDisplayName() const;
+
+	/**
+	 * Get all known key mappings for this profile.
+	 *
+	 * This returns a map of "Action Name" -> Mappings to that action
+	 */
+	UFUNCTION(BlueprintCallable, BlueprintPure, Category="Enhanced Input|User Settings")
+	const TMap<FName, FKeyMappingRow>& GetPlayerMappedActions() const;
+
+	/** Get all the key mappings associated with the given action name on this profile */
+	const FKeyMappingRow* FindKeyMappingRow(const FName InActionName) const;
+
+	/** A helper function to print out all the current profile settings to the log. */
+	UFUNCTION(BlueprintCallable, Category="Enhanced Input|User Settings")
+	virtual void DumpProfileToLog() const;
+
+	/** Returns a string that can be used to debug the current key mappings.  */
+	UFUNCTION(BlueprintCallable, Category="Enhanced Input|User Settings")
+	virtual FString ToString() const;
+	
+protected:
+
+	/** Returns a pointer to the player key mapping that fits with the given arguments. Returns null if none exist. */
+	virtual FPlayerKeyMapping* FindKeyMapping(const FMapPlayerKeyArgs& InArgs) const;
+	
+	/**
+	 * Resets all the key mappings in this profile to their default value from their Input Mapping Context.
+	 *
+	 * This function will only ever be called by UEnhancedInputUserSettings
+	 */
+	virtual void ResetToDefault();
+
+	/**
+	 * Equips the current key profile. This will always be called after the previous key profile's UnEquip function.
+	*  Make any changes to the given Enhanced Player input object that may be necessary for
+	 * your custom profile.
+	 * 
+	 * This function will only ever be called by UEnhancedInputUserSettings::SetKeyProfile.
+	 */
+	virtual void EquipProfile();
+
+	/**
+	 * UnEquips the current profile. Make any changes to the given Enhanced Player input object that may be necessary for
+	 * your custom profile.
+	 * 
+	 * This function will only ever be called by UEnhancedInputUserSettings::SetKeyProfile
+	 */
+	virtual void UnEquipProfile();
+	
+	/** The ID of this profile. This can be used by each Key Mapping to filter down which profile is required for it be equipped. */
+	UPROPERTY(BlueprintReadOnly, SaveGame, EditAnywhere, Category="Enhanced Input|User Settings")
+	FGameplayTag ProfileIdentifier;
+	
+	/** The localized display name of this profile */
+	UPROPERTY(BlueprintReadWrite, SaveGame, EditAnywhere, Category="Enhanced Input|User Settings")
+	FText DisplayName;
+	
+	/** A map of "Action Name" to all key mappings associated with it */
+	UPROPERTY(BlueprintReadOnly, SaveGame, EditAnywhere, Category="Enhanced Input|User Settings")
+	TMap<FName, FKeyMappingRow> PlayerMappedKeys;
+};
+
+/** Arguments that can be used when creating a new mapping profile */
+USTRUCT(BlueprintType)
+struct FPlayerMappableKeyProfileCreationArgs
+{
+	GENERATED_BODY()
+	
+	FPlayerMappableKeyProfileCreationArgs();
+	
+	UPROPERTY(BlueprintReadWrite, EditAnywhere, Category="Enhanced Input|User Settings")
+	TSubclassOf<UEnhancedPlayerMappableKeyProfile> ProfileType;
+	
+	/** The uniqiue identifier that this profile should have */
+	UPROPERTY(BlueprintReadWrite, EditAnywhere, Category="Enhanced Input|User Settings")
+	FGameplayTag ProfileIdentifier;
+	
+	/** The display name of this profile */
+	UPROPERTY(BlueprintReadWrite, EditAnywhere, Category="Enhanced Input|User Settings")
+	FText DisplayName;
+
+	UPROPERTY(BlueprintReadWrite, EditAnywhere, Category="Enhanced Input|User Settings")
+	uint8 bSetAsCurrentProfile : 1;
+};
+
+/**
+ * The Enhanced Input User Settings class is a place where you can put all of your Input Related settings
+ * that you want your user to be able to change. Things like their key mappings, aim sensativity, accessibility
+ * settings, etc. This also provies a Registration point for Input Mappings Contexts (IMC) from possibly unloaded
+ * plugins (i.e. Game Feature Plugins). You can register your IMC from a Game Feature Action plugin here, and then
+ * have access to all the key mappings available. This is very useful for building settings screens because you can
+ * now access all the mappings in your game, even if the entire plugin isn't loaded yet. 
+ *
+ * The user settings are stored on each UEnhancedPlayerInput object, so each instance of the settings can represent
+ * a single User or Local Player.
+ *
+ * To customize this for your game, you can create a subclass of it and change the "UserSettingsClass" in the
+ * Enhanced Input Project Settings.
+ */
+UCLASS(config=GameUserSettings, DisplayName="Enhanced Input User Settings (Experimental)", Category="Enhanced Input|User Settings")
+class ENHANCEDINPUT_API UEnhancedInputUserSettings : public USaveGame
+{
+	GENERATED_BODY()
+
+public:
+
+	//~ Begin UObject Interface
+	virtual void Serialize(FArchive& Ar) override;
+	//~ End UObject Interface
+	
+	/** Loads or creates new user settings for the owning local player of the given player input */
+	static UEnhancedInputUserSettings* LoadOrCreateSettings(UEnhancedPlayerInput* PlayerInput);
+	virtual void Initialize(UEnhancedPlayerInput* PlayerInput);
+
+	UFUNCTION(BlueprintCallable, Category="Enhanced Input|User Settings")
+	virtual void ApplySettings();
+
+	UFUNCTION(BlueprintCallable, Category="Enhanced Input|User Settings")
+	virtual void SaveSettings();
+	
+	UEnhancedPlayerInput* GetPlayerInput() const;
+	ULocalPlayer* GetLocalPlayer() const;
+	APlayerController* GetPlayerController() const;
+	
+	/** Fired when the user settings have changed, such as their key mappings. */
+	DECLARE_DYNAMIC_MULTICAST_DELEGATE_OneParam(FEnhancedInputUserSettingsChanged, UEnhancedInputUserSettings*, Settings);
+	FEnhancedInputUserSettingsChanged OnSettingsChanged;
+
+	// Remappable keys API
+
+	/**
+	 * Sets the player mapped key for this action on the current key profile.
+	 */
+	UFUNCTION(BlueprintCallable, Category="Enhanced Input|User Settings", meta = (AutoCreateRefTerm = "FailureReason"))
+	virtual void MapPlayerKey(const FMapPlayerKeyArgs& InArgs, FGameplayTagContainer& FailureReason);
+
+	/** Unmap what is currently mapped to the given action in the given slot */
+	UFUNCTION(BlueprintCallable, Category="Enhanced Input|User Settings", meta = (AutoCreateRefTerm = "FailureReason"))
+	virtual void UnMapPlayerKey(const FMapPlayerKeyArgs& InArgs, FGameplayTagContainer& FailureReason);
+
+	UFUNCTION(BlueprintCallable, Category="Enhanced Input|User Settings")
+	virtual const TSet<FPlayerKeyMapping>& FindMappingsForAction(const FName ActionName) const;
+
+	/** Returns the current player key mapping for the given action in the given slot */
+	virtual const FPlayerKeyMapping* FindCurrentMappingForSlot(const FName ActionName, const EPlayerMappableKeySlot InSlot) const;
+	
+	// Modifying key profile
+
+	/** Fired when you equip a different key profile  */
+	DECLARE_DYNAMIC_MULTICAST_DELEGATE_OneParam(FMappableKeyProfileChanged, const UEnhancedPlayerMappableKeyProfile*, NewProfile);
+	FMappableKeyProfileChanged OnKeyProfileChanged;
+	
+	/**
+	 * Changes the currently active key profile to the one with the given name. Returns true if the operation was successful.
+	 */
+	UFUNCTION(BlueprintCallable, Category="Enhanced Input|User Settings", meta=(ReturnDisplayName = "Was Successful"))
+	virtual bool SetKeyProfile(const FGameplayTag& InProfileId);
+	
+	/** Gets the currently selected key profile */
+	UFUNCTION(BlueprintCallable, BlueprintPure, Category="Enhanced Input|User Settings")
+	virtual const FGameplayTag& GetCurrentKeyProfileIdentifier() const;
+
+	/** Get the current key profile that the user has set */
+	UFUNCTION(BlueprintCallable, BlueprintPure, Category="Enhanced Input|User Settings")
+	UEnhancedPlayerMappableKeyProfile* GetCurrentKeyProfile() const;
+
+	template<class T>
+	T* GetCurrentKeyProfile() const
+	{
+		return Cast<T>(GetCurrentKeyProfile());
+	}
+
+	/**
+	 * Creates a new profile with this name and type.
+	 */
+	UFUNCTION(BlueprintCallable, Category="Enhanced Input|User Settings", meta = (DeterminesOutputType = "InArgs.ProfileType"))
+	virtual UEnhancedPlayerMappableKeyProfile* CreateNewKeyProfile(const FPlayerMappableKeyProfileCreationArgs& InArgs);
+
+	/** Returns the key profile with the given name if one exists. Null if one doesn't exist */
+	UFUNCTION(BlueprintCallable, BlueprintPure, Category="Enhanced Input|User Settings")
+	virtual UEnhancedPlayerMappableKeyProfile* GetKeyProfileWithIdentifier(const FGameplayTag& ProfileId) const;
+
+	template<class T>
+	T* GetKeyProfileWithIdentifier(const FGameplayTag& ProfileId) const
+	{
+		return Cast<T>(GetKeyProfileWithIdentifier(ProfileId));
+	}
+
+	// Registering input mapping contexts for access to them from your UI,
+	// even if they are froma plugin
+
+	/** Fired when a new input mapping context is registered. Useful if you need to update your UI */
+	DECLARE_DYNAMIC_MULTICAST_DELEGATE_OneParam(FMappingContextRegisteredWithSettings, const UInputMappingContext*, IMC);
+	FMappingContextRegisteredWithSettings OnMappingContextRegistered;
+
+	/**
+	 * Registers this mapping context with the user settings. This will iterate all the key mappings
+	 * in the context and create an initial Player Mappable Key for every mapping that is marked as mappable.
+	 */
+	UFUNCTION(BlueprintCallable, Category="Enhanced Input|User Settings")
+	bool RegisterInputMappingContext(UInputMappingContext* IMC);
+
+	/** Registers multiple mapping contexts with the settings */
+	UFUNCTION(BlueprintCallable, Category="Enhanced Input|User Settings")
+	bool RegisterInputMappingContexts(const TSet<UInputMappingContext*>& MappingContexts);
+
+	/** Removes this mapping context from the registered mapping contexts */
+	UFUNCTION(BlueprintCallable, Category="Enhanced Input|User Settings")
+	bool UnregisterInputMappingContext(const UInputMappingContext* IMC);
+
+	/** Removes miltiple mapping contexts from the registered mapping contexts */
+	UFUNCTION(BlueprintCallable, Category="Enhanced Input|User Settings")
+	bool UnregisterInputMappingContexts(const TSet<UInputMappingContext*>& MappingContexts);
+
+	/** Gets all the currently registered mapping contexts with the settings */
+	const TSet<TObjectPtr<UInputMappingContext>>& GetRegisteredInputMappingContexts() const;
+
+	/** Returns true if this mapping context is currently registered with the settings */
+	UFUNCTION(BlueprintCallable, BlueprintPure, Category="Enhanced Input|User Settings")
+	bool IsMappingContextRegistered(const UInputMappingContext* IMC) const;
+	
+protected:
+
+	/** The current key profile that is equipped by the user. */
+	UPROPERTY(SaveGame)
+	FGameplayTag CurrentProfileIdentifier;
+	
+	/**
+	 * All of the known Key Profiles for this user, including the currently active profile.
+	 */
+	UPROPERTY(SaveGame)
+	TMap<FGameplayTag, TObjectPtr<UEnhancedPlayerMappableKeyProfile>> SavedKeyProfiles;
+	
+	/** The owning Player Input object of these settings */
+	UPROPERTY(Transient)
+	TObjectPtr<UEnhancedPlayerInput> OwningPlayerInput;
+	
+	/**
+	 * Set of currently registered input mapping contexts that may not be currently
+	 * active on the user, but you want to track for creating a menu for key mappings.
+	 */
+	UPROPERTY(Transient)
+	TSet<TObjectPtr<UInputMappingContext>> RegisteredMappingContexts;
+};
