Fix python scripting expose issue

#rnx
#pf skip

diff --git a/Engine/Plugins/EnhancedInput/Source/EnhancedInput/Public/PlayerMappableKeySlot.h b/Engine/Plugins/EnhancedInput/Source/EnhancedInput/Public/PlayerMappableKeySlot.h
--- a/Engine/Plugins/EnhancedInput/Source/EnhancedInput/Public/PlayerMappableKeySlot.h
+++ b/Engine/Plugins/EnhancedInput/Source/EnhancedInput/Public/PlayerMappableKeySlot.h
@@ -8,7 +8,7 @@
 * Explicitly identifies the slot for a player mappable key
 * Experimental: Do not count on long term support for this structure.
 */
-USTRUCT(BlueprintType, DisplayName="Player Mappable Key Slot (Experimental)")
+USTRUCT(BlueprintType, DisplayName="Player Mappable Key Slot (Experimental)", meta = (ScriptName = "PlayerMappableKeySlotData"))
 struct ENHANCEDINPUT_API FPlayerMappableKeySlot
 {
 	GENERATED_BODY()
