Add some accessor functions for various commonly needed cases to the user settings.

#rnx

diff --git a/Engine/Plugins/EnhancedInput/Source/EnhancedInput/Private/UserSettings/EnhancedInputUserSettings.cpp b/Engine/Plugins/EnhancedInput/Source/EnhancedInput/Private/UserSettings/EnhancedInputUserSettings.cpp
--- a/Engine/Plugins/EnhancedInput/Source/EnhancedInput/Private/UserSettings/EnhancedInputUserSettings.cpp
+++ b/Engine/Plugins/EnhancedInput/Source/EnhancedInput/Private/UserSettings/EnhancedInputUserSettings.cpp
@@ -138,9 +138,12 @@ const FName FPlayerKeyMapping::GetActionName() const
 const FText& FPlayerKeyMapping::GetDisplayName() const
 {
 	// Just in case the display name is empty on this mapping, see if we can fall back to the original mapping copy's display name
-	if (DisplayName.IsEmpty() && !OriginalMappingCopy.PlayerMappableOptions.DisplayName.IsEmpty())
+	if (DisplayName.IsEmpty())
 	{
-		return OriginalMappingCopy.PlayerMappableOptions.DisplayName;
+		if (const UPlayerMappableKeySettings* Settings = OriginalMappingCopy.GetPlayerMappableKeySettings())
+		{
+			return Settings->DisplayName;
+		}
 	}
 	return DisplayName;
 }
@@ -237,6 +240,22 @@ const TMap<FName, FKeyMappingRow>& UEnhancedPlayerMappableKeyProfile::GetPlayerM
 	return PlayerMappedKeys;
 }
 
+void UEnhancedPlayerMappableKeyProfile::ResetActionMappingsToDefault(const FName InActionName)
+{
+	if (FKeyMappingRow* MappingRow = FindKeyMappingRowMutable(InActionName))
+	{
+		for (FPlayerKeyMapping& Mapping : MappingRow->Mappings)
+		{
+			Mapping.ResetToDefault();
+		}
+	}
+}
+
+FKeyMappingRow* UEnhancedPlayerMappableKeyProfile::FindKeyMappingRowMutable(const FName InActionName)
+{
+	return const_cast<FKeyMappingRow*>(FindKeyMappingRow(InActionName));
+}
+
 const FKeyMappingRow* UEnhancedPlayerMappableKeyProfile::FindKeyMappingRow(const FName InActionName) const
 {
 	return PlayerMappedKeys.Find(InActionName);
@@ -266,6 +285,45 @@ FString UEnhancedPlayerMappableKeyProfile::ToString() const
 	return Builder.ToString();
 }
 
+int32 UEnhancedPlayerMappableKeyProfile::GetKeysMappedToAction(const FName ActionName, TArray<FKey>& OutKeys) const
+{
+	OutKeys.Reset();
+	
+	if (const FKeyMappingRow* MappingRow = FindKeyMappingRow(ActionName))
+	{
+		for (const FPlayerKeyMapping& Mapping : MappingRow->Mappings)
+		{
+			OutKeys.Add(Mapping.GetCurrentKey());
+		}
+	}
+	else
+	{
+		UE_LOG(LogEnhancedInput, Warning, TEXT("Player Mappable Key Profile '%s' doesn't have any mappings for action '%s'"), *ProfileIdentifier.ToString(), *ActionName.ToString());
+	}
+	return OutKeys.Num();
+}
+
+int32 UEnhancedPlayerMappableKeyProfile::GetActionsMappedToKey(const FKey& InKey, TArray<FName>& OutMappedActionNames) const
+{
+	OutMappedActionNames.Reset();
+
+	for (const TPair<FName, FKeyMappingRow>& Pair : PlayerMappedKeys)
+	{
+		for (const FPlayerKeyMapping& Mapping : Pair.Value.Mappings)
+		{
+			if (Mapping.GetCurrentKey() == InKey)
+			{
+				// We know that this action has the key mapped to it, so there is no need to continue checking
+				// the rest of it's mappings
+				OutMappedActionNames.Add(Pair.Key);
+				break;
+			}
+		}
+	}
+
+	return OutMappedActionNames.Num();
+}
+
 void UEnhancedPlayerMappableKeyProfile::Serialize(FArchive& Ar)
 {
 	// See note in header!
@@ -288,6 +346,18 @@ FPlayerKeyMapping* UEnhancedPlayerMappableKeyProfile::FindKeyMapping(const FMapP
 	return nullptr;
 }
 
+void UEnhancedPlayerMappableKeyProfile::K2_FindKeyMapping(FPlayerKeyMapping& OutKeyMapping, const FMapPlayerKeyArgs& InArgs) const
+{
+	if (FPlayerKeyMapping* FoundMapping = FindKeyMapping(InArgs))
+	{
+		OutKeyMapping = *FoundMapping;
+	}
+	else
+	{
+		OutKeyMapping = FPlayerKeyMapping::InvalidMapping;
+	}
+}
+
 ///////////////////////////////////////////////////////////
 // UEnhancedInputUserSettings
 
@@ -299,7 +369,7 @@ UEnhancedInputUserSettings* UEnhancedInputUserSettings::LoadOrCreateSettings(UEn
 
 	if (!LocalPlayer)
 	{
-		UE_LOG(LogEnhancedInput, Error, TEXT("Unable to determine an owning Local Player for the given Enhanced Player Input object"));
+		UE_LOG(LogEnhancedInput, Log, TEXT("Unable to determine an owning Local Player for the given Enhanced Player Input object"));
 		return nullptr;
 	}
 	
@@ -643,6 +713,11 @@ FPlayerMappableKeyProfileCreationArgs::FPlayerMappableKeyProfileCreationArgs()
 {
 }
 
+const TMap<FGameplayTag, TObjectPtr<UEnhancedPlayerMappableKeyProfile>>& UEnhancedInputUserSettings::GetAllSavedKeyProfiles() const
+{
+	return SavedKeyProfiles;
+}
+
 UEnhancedPlayerMappableKeyProfile* UEnhancedInputUserSettings::CreateNewKeyProfile(const FPlayerMappableKeyProfileCreationArgs& InArgs)
 {
 	if (!InArgs.ProfileType)
@@ -763,7 +838,12 @@ bool UEnhancedInputUserSettings::RegisterInputMappingContext(UInputMappingContex
 			FPlayerKeyMapping PlayerMappingData = {};
 			PlayerMappingData.ActionName = ActionName;
 			PlayerMappingData.DefaultKey = KeyMapping.Key;
-			PlayerMappingData.DisplayName = KeyMapping.PlayerMappableOptions.DisplayName;
+			
+			if (UPlayerMappableKeySettings* Settings = KeyMapping.GetPlayerMappableKeySettings())
+			{
+				PlayerMappingData.DisplayName = Settings->DisplayName;
+			}
+			
 			PlayerMappingData.OriginalMappingCopy = KeyMapping;
 
 			// By default, the slot will be determined by how many mappings this action has already.
diff --git a/Engine/Plugins/EnhancedInput/Source/EnhancedInput/Public/UserSettings/EnhancedInputUserSettings.h b/Engine/Plugins/EnhancedInput/Source/EnhancedInput/Public/UserSettings/EnhancedInputUserSettings.h
--- a/Engine/Plugins/EnhancedInput/Source/EnhancedInput/Public/UserSettings/EnhancedInputUserSettings.h
+++ b/Engine/Plugins/EnhancedInput/Source/EnhancedInput/Public/UserSettings/EnhancedInputUserSettings.h
@@ -219,6 +219,13 @@ public:
 	UFUNCTION(BlueprintCallable, BlueprintPure, Category="Enhanced Input|User Settings")
 	const TMap<FName, FKeyMappingRow>& GetPlayerMappedActions() const;
 
+	/** Resets every player key mapping to this action back to it's default value */
+	UFUNCTION(BlueprintCallable, Category="Enhanced Input|User Settings")
+	void ResetActionMappingsToDefault(const FName InActionName);
+	
+	/** Get all the key mappings associated with the given action name on this profile */
+	FKeyMappingRow* FindKeyMappingRowMutable(const FName InActionName);
+
 	/** Get all the key mappings associated with the given action name on this profile */
 	const FKeyMappingRow* FindKeyMappingRow(const FName InActionName) const;
 
@@ -230,11 +237,31 @@ public:
 	UFUNCTION(BlueprintCallable, Category="Enhanced Input|User Settings")
 	virtual FString ToString() const;
 	
-protected:
+	/**
+	 * Returns all FKey's bound to the given Action Name on this profile.
+	 *
+	 * Returns the number of keys mapped to this action
+	 */
+	UFUNCTION(BlueprintCallable, Category="Enhanced Input|User Settings", meta = (ReturnDisplayName = "Number of keys"))
+	virtual int32 GetKeysMappedToAction(const FName ActionName, /*OUT*/ TArray<FKey>& OutKeys) const;
+
+	/**
+	 * Populates the OutMappedActionNames with every action on this profile that has a mapping to the given key.
+	 *
+	 * Returns the number of actions mapped to this key
+	 */
+	UFUNCTION(BlueprintCallable, Category="Enhanced Input|User Settings", meta = (ReturnDisplayName = "Number of actions"))
+	virtual int32 GetActionsMappedToKey(const FKey& InKey, /*OUT*/ TArray<FName>& OutMappedActionNames) const;
 
 	/** Returns a pointer to the player key mapping that fits with the given arguments. Returns null if none exist. */
 	virtual FPlayerKeyMapping* FindKeyMapping(const FMapPlayerKeyArgs& InArgs) const;
 	
+	UFUNCTION(BlueprintCallable, Category="Enhanced Input|User Settings", meta=(DisplayName="Find Key Mapping", AutoCreateRefTerm="OutKeyMapping"))
+	void K2_FindKeyMapping(FPlayerKeyMapping& OutKeyMapping, const FMapPlayerKeyArgs& InArgs) const;
+	
+protected:
+
+	
 	/**
 	 * Resets all the key mappings in this profile to their default value from their Input Mapping Context.
 	 *
@@ -382,6 +409,9 @@ public:
 		return Cast<T>(GetCurrentKeyProfile());
 	}
 
+	/** Returns all player saved key profiles */
+	const TMap<FGameplayTag, TObjectPtr<UEnhancedPlayerMappableKeyProfile>>& GetAllSavedKeyProfiles() const;
+
 	/**
 	 * Creates a new profile with this name and type.
 	 */
