Spawn a new user settings object for each player owned EnhancedPlayerInput, making it accessible to our Enhanced Input subsystem.

#rnx

diff --git a/Engine/Plugins/EnhancedInput/Source/EnhancedInput/Private/EnhancedInputSubsystemInterface.cpp b/Engine/Plugins/EnhancedInput/Source/EnhancedInput/Private/EnhancedInputSubsystemInterface.cpp
--- a/Engine/Plugins/EnhancedInput/Source/EnhancedInput/Private/EnhancedInputSubsystemInterface.cpp
+++ b/Engine/Plugins/EnhancedInput/Source/EnhancedInput/Private/EnhancedInputSubsystemInterface.cpp
@@ -4,6 +4,8 @@
 
 #include "EnhancedInputModule.h"
 #include "EnhancedInputPlatformSettings.h"
+#include "UserSettings/EnhancedInputUserSettings.h"
+#include "EnhancedInputDeveloperSettings.h"
 #include "GameFramework/PlayerController.h"
 #include "HAL/IConsoleManager.h"
 #include "InputMappingContext.h"
@@ -27,6 +29,48 @@ static FAutoConsoleVariableRef GCVarGlobalAxisConfigMode(
 	TEXT("Whether or not to apply Global Axis Config settings. 0 = Default (Mouse Only), 1 = All, 2 = None")
 );
 
+UEnhancedInputUserSettings* IEnhancedInputSubsystemInterface::GetUserSettings() const
+{
+	if (const UEnhancedPlayerInput* PlayerInput = GetPlayerInput())
+	{
+		return PlayerInput->GetUserSettings();
+	}
+	return nullptr;
+}
+
+void IEnhancedInputSubsystemInterface::BindUserSettingDelegates()
+{
+	UEnhancedInputUserSettings* Settings = GetUserSettings();
+	if (!Settings)
+	{
+		UE_LOG(LogEnhancedInput, Error, TEXT("Unable to get the user settings object!"));
+		return;
+	}
+
+	// There is no need to bind to any delegates if the setting is turned off. We shouldn't even get here,
+	// but do this in case someone implements this interface
+	if (!GetDefault<UEnhancedInputDeveloperSettings>()->bEnableUserSettings)
+	{
+		UE_LOG(LogEnhancedInput, Error, TEXT("Attempting to bind to user settings delegates but they are disabled in UEnhancedInputDeveloperSettings!"));
+		return;
+	}
+
+	Settings->OnSettingsChanged.AddDynamic(this, &IEnhancedInputSubsystemInterface::OnUserSettingsChanged);
+	Settings->OnKeyProfileChanged.AddDynamic(this, &IEnhancedInputSubsystemInterface::OnUserKeyProfileChanged);
+}
+
+void IEnhancedInputSubsystemInterface::OnUserSettingsChanged(UEnhancedInputUserSettings* Settings)
+{
+	// We want to rebuild our control mappings whenever a setting has changed
+	RequestRebuildControlMappings();
+}
+
+void IEnhancedInputSubsystemInterface::OnUserKeyProfileChanged(const UEnhancedPlayerMappableKeyProfile* InNewProfile)
+{
+	// We want to rebuild our control mappings whenever a setting has changed
+	RequestRebuildControlMappings();
+}
+
 void IEnhancedInputSubsystemInterface::InjectInputForAction(const UInputAction* Action, FInputActionValue RawValue, const TArray<UInputModifier*>& Modifiers, const TArray<UInputTrigger*>& Triggers)
 {
 	if(UEnhancedPlayerInput* PlayerInput = GetPlayerInput())
diff --git a/Engine/Plugins/EnhancedInput/Source/EnhancedInput/Private/EnhancedPlayerInput.cpp b/Engine/Plugins/EnhancedInput/Source/EnhancedInput/Private/EnhancedPlayerInput.cpp
--- a/Engine/Plugins/EnhancedInput/Source/EnhancedInput/Private/EnhancedPlayerInput.cpp
+++ b/Engine/Plugins/EnhancedInput/Source/EnhancedInput/Private/EnhancedPlayerInput.cpp
@@ -8,6 +8,8 @@
 #include "GameFramework/WorldSettings.h"
 #include "EnhancedInputDeveloperSettings.h"
 #include "InputMappingContext.h"
+#include "UserSettings/EnhancedInputUserSettings.h"
+#include "EnhancedInputModule.h"
 
 #include UE_INLINE_GENERATED_CPP_BY_NAME(EnhancedPlayerInput)
 
@@ -42,6 +44,46 @@ UEnhancedPlayerInput::UEnhancedPlayerInput()
 	}
 }
 
+void UEnhancedPlayerInput::PostInitProperties()
+{
+	Super::PostInitProperties();
+
+	if (HasAnyFlags(RF_ClassDefaultObject))
+	{
+		return;
+	}
+	
+	InitalizeUserSettings();
+}
+
+void UEnhancedPlayerInput::InitalizeUserSettings()
+{
+	// We don't want to be re-creating any user settings, they should share the lifetime with the player input object
+	if (!ensureMsgf(!UserSettings, TEXT("Attempting to initalize the User Settings, but they have already been created!")))
+	{
+		return;
+	}
+	
+	const UEnhancedInputDeveloperSettings* DevSettings = GetDefault<UEnhancedInputDeveloperSettings>();
+	
+	// These are player settings, so only player owned input components can have them.
+	const bool bHasValidPlayerController = GetOuterAPlayerController() != nullptr;
+
+	if (bHasValidPlayerController && DevSettings->bEnableUserSettings && !IsDefaultSubobject())
+	{
+		// Spawn a new instance of the specified user settings class. This is a "NoClear" UPROPERTY so it should always be set
+		const UClass* SettingsClass = DevSettings->UserSettingsClass.Get();
+		if (ensure(SettingsClass))
+		{
+			UserSettings = UEnhancedInputUserSettings::LoadOrCreateSettings(this);
+		}
+	}
+	else
+	{
+		UE_LOG(LogEnhancedInput, Log, TEXT("bEnableUserSettings is set to false, skipping creation of UEnhancedInputUserSettings!"));
+	}
+}
+
 // NOTE: Enum order represents firing priority(lowest to highest) and is important as multiple keys bound to the same action may generate differing trigger event states.
 enum class ETriggerEventInternal : uint8
 {
@@ -186,7 +228,7 @@ void UEnhancedPlayerInput::ProcessActionMappingEvent(TObjectPtr<const UInputActi
 		// Apply modifications to the raw value
 		EInputActionValueType ValueType = ActionData.Value.GetValueType();
 		FInputActionValue ModifiedValue = ApplyModifiers(Modifiers, FInputActionValue(ValueType, RawKeyValue.Get<FVector>()), DeltaTime);
-		//UE_CLOG(RawKeyValue.GetMagnitudeSq(), LogTemp, Warning, TEXT("Modified %s -> %s"), *RawKeyValue.ToString(), *ModifiedValue.ToString());
+		//UE_CLOG(RawKeyValue.GetMagnitudeSq(), LogEnhancedInput, Warning, TEXT("Modified %s -> %s"), *RawKeyValue.ToString(), *ModifiedValue.ToString());
 
 		// Derive an initial trigger state for this mapping using all applicable triggers
 		ETriggerState CalcedState = TriggerStateTracker.EvaluateTriggers(this, Triggers, ModifiedValue, DeltaTime);
@@ -253,6 +295,11 @@ float UEnhancedPlayerInput::GetEffectiveTimeDilation() const
 	return 1.0f;
 }
 
+UEnhancedInputUserSettings* UEnhancedPlayerInput::GetUserSettings() const
+{
+	return UserSettings;
+}
+
 void UEnhancedPlayerInput::ProcessInputStack(const TArray<UInputComponent*>& InputComponentStack, const float DeltaTime, const bool bGamePaused)
 {
 	// We need to grab the down states of all keys before calling Super::ProcessInputStack as it will leave bDownPrevious in the same state as bDown (i.e. this frame, not last).
@@ -302,7 +349,7 @@ void UEnhancedPlayerInput::ProcessInputStack(const TArray<UInputComponent*>& Inp
 
 		FKeyState* KeyState = GetKeyState(Mapping.Key);
 		FVector RawKeyValue = KeyState ? KeyState->RawValue : FVector::ZeroVector;
-		//UE_CLOG(RawKeyValue.SizeSquared(), LogTemp, Warning, TEXT("Key %s - state %s"), *Mapping.Key.GetDisplayName().ToString(), *RawKeyValue.ToString());
+		//UE_CLOG(RawKeyValue.SizeSquared(), LogEnhancedInput, Warning, TEXT("Key %s - state %s"), *Mapping.Key.GetDisplayName().ToString(), *RawKeyValue.ToString());
 
 		// Should this key be ignored because it was down during a context switch?
 		// If so, check if it is back up, otherwise ignore it.
@@ -503,7 +550,7 @@ void UEnhancedPlayerInput::ProcessInputStack(const TArray<UInputComponent*>& Inp
 						bCanTrigger = !TriggeredActionsThisTick.Contains(DepAction.SourceAction);
 						if(!bCanTrigger)
 						{
-							UE_LOG(LogTemp, Warning, TEXT("'%s' action was cancelled, its dependant on '%s'"), *DelegateAction->GetName(), *DepAction.SourceAction->GetName());
+							UE_LOG(LogEnhancedInput, Warning, TEXT("'%s' action was cancelled, its dependant on '%s'"), *DelegateAction->GetName(), *DepAction.SourceAction->GetName());
 						}
 						break;
 					}
diff --git a/Engine/Plugins/EnhancedInput/Source/EnhancedInput/Private/UserSettings/EnhancedInputUserSettings.cpp b/Engine/Plugins/EnhancedInput/Source/EnhancedInput/Private/UserSettings/EnhancedInputUserSettings.cpp
--- a/Engine/Plugins/EnhancedInput/Source/EnhancedInput/Private/UserSettings/EnhancedInputUserSettings.cpp
+++ b/Engine/Plugins/EnhancedInput/Source/EnhancedInput/Private/UserSettings/EnhancedInputUserSettings.cpp
@@ -50,7 +50,7 @@ namespace UE::EnhancedInput
 	static void DumpAllKeyProfilesToLog(const TArray<FString>& Args)
 	{
 		// Dump every local player subsystem's logs
-		/*UEnhancedInputLibrary::ForEachSubsystem([Args](IEnhancedInputSubsystemInterface* Subsystem)
+		UEnhancedInputLibrary::ForEachSubsystem([Args](IEnhancedInputSubsystemInterface* Subsystem)
 		{
 			if (const UEnhancedInputUserSettings* Settings = Subsystem->GetUserSettings())
 			{
@@ -59,7 +59,7 @@ namespace UE::EnhancedInput
 					Profile->DumpProfileToLog();	
 				}
 			}
-		});*/
+		});
 	}
 	
 	static FAutoConsoleCommand ConsoelCommandDumpProfileToLog(
diff --git a/Engine/Plugins/EnhancedInput/Source/EnhancedInput/Public/EnhancedInputSubsystemInterface.h b/Engine/Plugins/EnhancedInput/Source/EnhancedInput/Public/EnhancedInputSubsystemInterface.h
--- a/Engine/Plugins/EnhancedInput/Source/EnhancedInput/Public/EnhancedInputSubsystemInterface.h
+++ b/Engine/Plugins/EnhancedInput/Source/EnhancedInput/Public/EnhancedInputSubsystemInterface.h
@@ -21,6 +21,8 @@ class UInputModifier;
 class UInputTrigger;
 class UPlayerMappableInputConfig;
 class UWorldSubsystem;
+class UEnhancedInputUserSettings;
+class UEnhancedPlayerMappableKeyProfile;
 
 // Subsystem interface
 UINTERFACE(MinimalAPI, meta=(CannotImplementInterfaceInBlueprint))
@@ -71,6 +73,28 @@ public:
 
 	virtual UEnhancedPlayerInput* GetPlayerInput() const = 0;
 
+	UFUNCTION(BlueprintCallable, BlueprintPure, Category="Enhanced Input|User Settings")
+	virtual UEnhancedInputUserSettings* GetUserSettings() const;
+
+	template<class T>
+	inline T* GetUserSettings() const
+	{
+		return Cast<T>(GetUserSettings());
+	}
+
+protected:
+	
+	/** Binds to any delegates of interest on the UEnhancedInputUserSettings if they are enabled in the developer settings. */
+	virtual void BindUserSettingDelegates();
+	
+	/** Callback for when any Enhanced Input user settings have been changed (a new key mapping for example) */
+	virtual void OnUserSettingsChanged(UEnhancedInputUserSettings* Settings);
+
+	/** A callback for when the user has applied a new mappable key profile. */
+	virtual void OnUserKeyProfileChanged(const UEnhancedPlayerMappableKeyProfile* InNewProfile);
+
+public:
+	
 	/**
 	 * Input simulation via injection. Runs modifiers and triggers delegates as if the input had come through the underlying input system as FKeys.
 	 * Applies action modifiers and triggers on top.
diff --git a/Engine/Plugins/EnhancedInput/Source/EnhancedInput/Public/EnhancedPlayerInput.h b/Engine/Plugins/EnhancedInput/Source/EnhancedInput/Public/EnhancedPlayerInput.h
--- a/Engine/Plugins/EnhancedInput/Source/EnhancedInput/Public/EnhancedPlayerInput.h
+++ b/Engine/Plugins/EnhancedInput/Source/EnhancedInput/Public/EnhancedPlayerInput.h
@@ -12,6 +12,7 @@ class UInputTrigger;
 enum class ETriggerEvent : uint8;
 enum class ETriggerState : uint8;
 struct FEnhancedActionKeyMapping;
+class UEnhancedInputUserSettings;
 
 // Internal representation containing event variants
 enum class ETriggerEventInternal : uint8;
@@ -48,6 +49,8 @@ public:
 
 	UEnhancedPlayerInput();
 
+	virtual void PostInitProperties() override;
+
 	/**
 	* Returns the action instance data for the given input action if there is any. Returns nullptr if the action is not available.
 	*/
@@ -68,8 +71,14 @@ public:
 
 	/** Returns the Time Dilation value that is currently effecting this input. */
 	float GetEffectiveTimeDilation() const;
+
+	/** Returns the current user settings for this player input. Will be null if the bEnableUserSettings flag is false on the developer settings */
+	UEnhancedInputUserSettings* GetUserSettings() const;
 	
 protected:
+
+	/** Create a new user settings object if it is enabled in the EI developer settings */
+	virtual void InitalizeUserSettings();
 	
 	// Causes key to be consumed if it is affecting an action.
 	virtual bool IsKeyHandledByAction(FKey Key) const override;
@@ -158,6 +167,14 @@ private:
 	 */
 	TArray<FDependentChordTracker> DependentChordActions;
 
+protected:
+	
+	/** The user settings for this subsystem used to store each user's input related settings */
+	UPROPERTY()
+	TObjectPtr<UEnhancedInputUserSettings> UserSettings;
+
+private:
+
 	/** The last time of the last frame that was processed in ProcessPlayerInput */
 	float LastFrameTime = 0.0f;
 
