Change the category on the player mappable settings so that it isnt needlessly nested twice

#pf skip
#rnx

diff --git a/Engine/Plugins/EnhancedInput/Source/EnhancedInput/Public/InputAction.h b/Engine/Plugins/EnhancedInput/Source/EnhancedInput/Public/InputAction.h
--- a/Engine/Plugins/EnhancedInput/Source/EnhancedInput/Public/InputAction.h
+++ b/Engine/Plugins/EnhancedInput/Source/EnhancedInput/Public/InputAction.h
@@ -88,7 +88,7 @@ private:
 	/**
 	* Holds setting information about this Action Input for setting screen and save purposes.
 	*/
-	UPROPERTY(EditAnywhere, Instanced, BlueprintReadWrite, Category = "Input|Settings", meta=(AllowPrivateAccess))
+	UPROPERTY(EditAnywhere, Instanced, BlueprintReadWrite, Category = "User Settings", meta=(AllowPrivateAccess))
 	TObjectPtr<UPlayerMappableKeySettings> PlayerMappableKeySettings;
 };
 
