Just some small organizational things to the Action key mapping

#rnx

diff --git a/Engine/Plugins/EnhancedInput/Source/EnhancedInput/Private/EnhancedActionKeyMapping.cpp b/Engine/Plugins/EnhancedInput/Source/EnhancedInput/Private/EnhancedActionKeyMapping.cpp
--- a/Engine/Plugins/EnhancedInput/Source/EnhancedInput/Private/EnhancedActionKeyMapping.cpp
+++ b/Engine/Plugins/EnhancedInput/Source/EnhancedInput/Private/EnhancedActionKeyMapping.cpp
@@ -21,7 +21,15 @@ UPlayerMappableKeySettings* FEnhancedActionKeyMapping::GetPlayerMappableKeySetti
 	return nullptr;
 }
 
-ENHANCEDINPUT_API FName FEnhancedActionKeyMapping::GetMappingName() const
+FEnhancedActionKeyMapping::FEnhancedActionKeyMapping(const UInputAction* InAction /*= nullptr*/, const FKey InKey /*= EKeys::Invalid*/)
+	: PlayerMappableOptions(InAction)
+	, Action(InAction)
+	, Key(InKey)
+	, bShouldBeIgnored(false)
+	, bIsPlayerMappable(false)
+{}
+
+FName FEnhancedActionKeyMapping::GetMappingName() const
 {
 	if (IsPlayerMappable())
 	{
@@ -156,4 +164,12 @@ EDataValidationResult FEnhancedActionKeyMapping::IsDataValid(TArray<FText>& Vali
 
 #endif // WITH_EDITOR
 
+bool FEnhancedActionKeyMapping::operator==(const FEnhancedActionKeyMapping& Other) const
+{
+	return (Action == Other.Action &&
+			Key == Other.Key &&
+			Triggers == Other.Triggers &&
+			Modifiers == Other.Modifiers);
+}
+
 #undef LOCTEXT_NAMESPACE
diff --git a/Engine/Plugins/EnhancedInput/Source/EnhancedInput/Public/EnhancedActionKeyMapping.h b/Engine/Plugins/EnhancedInput/Source/EnhancedInput/Public/EnhancedActionKeyMapping.h
--- a/Engine/Plugins/EnhancedInput/Source/EnhancedInput/Public/EnhancedActionKeyMapping.h
+++ b/Engine/Plugins/EnhancedInput/Source/EnhancedInput/Public/EnhancedActionKeyMapping.h
@@ -77,28 +77,35 @@ public:
  *
 **/
 USTRUCT(BlueprintType)
-struct FEnhancedActionKeyMapping
+struct ENHANCEDINPUT_API FEnhancedActionKeyMapping
 {
+	friend class UInputMappingContext;
+	friend class FEnhancedActionMappingCustomization;
+	
 	GENERATED_BODY()
 
+	FEnhancedActionKeyMapping(const UInputAction* InAction = nullptr, const FKey InKey = EKeys::Invalid);
+	
 	/**
 	* Returns the Player Mappable Key Settings owned by the Action Key Mapping or by the referenced Input Action, or nothing based of the Setting Behavior.
 	*/
-	ENHANCEDINPUT_API UPlayerMappableKeySettings* GetPlayerMappableKeySettings() const;
+	UPlayerMappableKeySettings* GetPlayerMappableKeySettings() const;
 
 	/**
 	 * Returns the name of the mapping based on setting behavior used. If no name is found in the Mappable Key Settings it will return the name set in Player Mappable Options if bIsPlayerMappable is true.
 	 */
-	ENHANCEDINPUT_API FName GetMappingName() const;
+	FName GetMappingName() const;
 
 	/**
 	* Returns true if this Action Key Mapping either holds a Player Mappable Key Settings or is set bIsPlayerMappable.
 	*/
-	ENHANCEDINPUT_API bool IsPlayerMappable() const;
+	bool IsPlayerMappable() const;
 
 #if WITH_EDITOR
 	EDataValidationResult IsDataValid(TArray<FText>& ValidationErrors);
 #endif
+	
+	bool operator==(const FEnhancedActionKeyMapping& Other) const;
 
 	/** Options for making this a player mappable keymapping */
 	UPROPERTY(EditAnywhere, BlueprintReadWrite, Category = "Input|PlayerMappable", meta = (editCondition = "bIsPlayerMappable", DisplayAfter = "bIsPlayerMappable"))
@@ -136,23 +143,10 @@ struct FEnhancedActionKeyMapping
 	UPROPERTY(Transient)
 	uint8 bShouldBeIgnored : 1;
 	
-	bool operator==(const FEnhancedActionKeyMapping& Other) const
-	{
-		return (Action == Other.Action &&
-				Key == Other.Key &&
-				Triggers == Other.Triggers &&
-				Modifiers == Other.Modifiers);
-	}
-
-	FEnhancedActionKeyMapping(const UInputAction* InAction = nullptr, const FKey InKey = EKeys::Invalid)
-		: PlayerMappableOptions(InAction)
-		, Action(InAction)
-		, Key(InKey)
-		, bShouldBeIgnored(false)
-		, bIsPlayerMappable(false)
-	{}
+	/** If true then this ActionKeyMapping will be exposed as a player mappable key */
+	UPROPERTY(EditAnywhere, Category = "Input|PlayerMappable")
+	uint8 bIsPlayerMappable : 1;
 
-	friend class FEnhancedActionMappingCustomization;
 
 protected:
 
@@ -168,10 +162,6 @@ protected:
 	UPROPERTY(EditAnywhere, Instanced, Category = "Input|Settings", meta = (EditCondition = "SettingBehavior == EPlayerMappableKeySettingBehaviors::OverrideSettings", DisplayAfter = "SettingBehavior"))
 	TObjectPtr<UPlayerMappableKeySettings> PlayerMappableKeySettings = nullptr;
 
-	/** If true then this ActionKeyMapping will be exposed as a player mappable key */
-	UPROPERTY(EditAnywhere, BlueprintReadWrite, Category = "Input|PlayerMappable")
-	uint8 bIsPlayerMappable : 1;
-
 public:
 
 	template<typename T = UPlayerMappableKeySettings> 
