Removed BlueprintPure from interface function declaration.  C# UHT was failing to signal that as an error while C++ UHT would.

diff --git a/Engine/Plugins/EnhancedInput/Source/EnhancedInput/Public/EnhancedInputSubsystemInterface.h b/Engine/Plugins/EnhancedInput/Source/EnhancedInput/Public/EnhancedInputSubsystemInterface.h
--- a/Engine/Plugins/EnhancedInput/Source/EnhancedInput/Public/EnhancedInputSubsystemInterface.h
+++ b/Engine/Plugins/EnhancedInput/Source/EnhancedInput/Public/EnhancedInputSubsystemInterface.h
@@ -73,7 +73,7 @@ public:
 
 	virtual UEnhancedPlayerInput* GetPlayerInput() const = 0;
 
-	UFUNCTION(BlueprintCallable, BlueprintPure, Category="Enhanced Input|User Settings")
+	UFUNCTION(BlueprintCallable, Category="Enhanced Input|User Settings")
 	virtual UEnhancedInputUserSettings* GetUserSettings() const;
 
 	template<class T>
