Add a helped functoin to the mapping row to check if it has any mappings. Export the mappable key creation struct to the EI api

#rnx
#pf skip

diff --git a/Engine/Plugins/EnhancedInput/Source/EnhancedInput/Private/UserSettings/EnhancedInputUserSettings.cpp b/Engine/Plugins/EnhancedInput/Source/EnhancedInput/Private/UserSettings/EnhancedInputUserSettings.cpp
--- a/Engine/Plugins/EnhancedInput/Source/EnhancedInput/Private/UserSettings/EnhancedInputUserSettings.cpp
+++ b/Engine/Plugins/EnhancedInput/Source/EnhancedInput/Private/UserSettings/EnhancedInputUserSettings.cpp
@@ -324,6 +324,11 @@ int32 UEnhancedPlayerMappableKeyProfile::GetActionsMappedToKey(const FKey& InKey
 	return OutMappedActionNames.Num();
 }
 
+bool FKeyMappingRow::HasAnyMappings() const
+{
+	return !Mappings.IsEmpty();
+}
+
 void UEnhancedPlayerMappableKeyProfile::Serialize(FArchive& Ar)
 {
 	// See note in header!
diff --git a/Engine/Plugins/EnhancedInput/Source/EnhancedInput/Public/UserSettings/EnhancedInputUserSettings.h b/Engine/Plugins/EnhancedInput/Source/EnhancedInput/Public/UserSettings/EnhancedInputUserSettings.h
--- a/Engine/Plugins/EnhancedInput/Source/EnhancedInput/Public/UserSettings/EnhancedInputUserSettings.h
+++ b/Engine/Plugins/EnhancedInput/Source/EnhancedInput/Public/UserSettings/EnhancedInputUserSettings.h
@@ -178,6 +178,9 @@ struct ENHANCEDINPUT_API FKeyMappingRow
 
 	UPROPERTY(VisibleAnywhere, SaveGame, BlueprintReadOnly, Category="Enhanced Input|User Settings")
 	TSet<FPlayerKeyMapping> Mappings;
+
+	/** Returns true if this row has any mappings in it */
+	bool HasAnyMappings() const;
 };
 
 
@@ -301,7 +304,7 @@ protected:
 
 /** Arguments that can be used when creating a new mapping profile */
 USTRUCT(BlueprintType)
-struct FPlayerMappableKeyProfileCreationArgs
+struct ENHANCEDINPUT_API FPlayerMappableKeyProfileCreationArgs
 {
 	GENERATED_BODY()
 	
