Implementation of using the Enhanced Input User Settings instead of the old "PlayerMappedKeys" property on the enhanced input subsystem

#rnx

diff --git a/Engine/Plugins/EnhancedInput/Source/EnhancedInput/Private/EnhancedInputSubsystemInterface.cpp b/Engine/Plugins/EnhancedInput/Source/EnhancedInput/Private/EnhancedInputSubsystemInterface.cpp
--- a/Engine/Plugins/EnhancedInput/Source/EnhancedInput/Private/EnhancedInputSubsystemInterface.cpp
+++ b/Engine/Plugins/EnhancedInput/Source/EnhancedInput/Private/EnhancedInputSubsystemInterface.cpp
@@ -746,24 +746,48 @@ void IEnhancedInputSubsystemInterface::RebuildControlMappings()
 	// Reset the tracking of dependant chord actions on the player input
 	PlayerInput->DependentChordActions.Reset();
 
+	UEnhancedInputUserSettings* UserSettings = GetUserSettings();
+	UEnhancedPlayerMappableKeyProfile* PlayerKeyProfile = UserSettings ? UserSettings->GetCurrentKeyProfile() : nullptr;
+
+	// An array of keys that are mapped to a given Action.
+	// This is populated by any player mapped keys if they exist, or the default mapping from
+	// an input mapping context.
+	TArray<FKey> MappedKeysToActionName;
+	
 	for (const TPair<TObjectPtr<const UInputMappingContext>, int32>& ContextPair : OrderedInputContexts)
 	{
 		// Don't apply context specific keys immediately, allowing multiple mappings to the same key within the same context if required.
 		TArray<FKey> ContextAppliedKeys;
-
 		const UInputMappingContext* MappingContext = ContextPair.Key;
 		TArray<FEnhancedActionKeyMapping> OrderedMappings = ReorderMappings(MappingContext->GetMappings(), PlayerInput->DependentChordActions);
 
 		for (FEnhancedActionKeyMapping& Mapping : OrderedMappings)
 		{
-			TArray<FKey> PlayerMappedKeys = GetAllPlayerMappedKeys(Mapping.GetMappingName());
-			if (PlayerMappedKeys.IsEmpty())
+			// Clear out mappings from the previous iteration
+			MappedKeysToActionName.Reset();
+
+			// True if there were any player mapped keys to this mapping and we are using those instead.
+			bool bIsPlayerMapping = false;
+			
+			// See if there are any player mapped keys to this action
+			if (PlayerKeyProfile && GetDefault<UEnhancedInputDeveloperSettings>()->bEnableUserSettings)
 			{
-				//If we didn't find any player mapped keys use the default key from the mapping.
-				PlayerMappedKeys.Add(Mapping.Key);
+				PlayerKeyProfile->GetKeysMappedToAction(Mapping.GetMappingName(), MappedKeysToActionName);
+				bIsPlayerMapping = !MappedKeysToActionName.IsEmpty();
+			}
+			else if (const TMap<FPlayerMappableKeySlot, FKey>* PlayerMappedKeySlots = PlayerMappedSettings.Find(Mapping.GetMappingName()))
+			{
+				PlayerMappedKeySlots->GenerateValueArray(MappedKeysToActionName);
+				bIsPlayerMapping = !MappedKeysToActionName.IsEmpty();
 			}
 
-			for (const FKey& PlayerMappedKey : PlayerMappedKeys)
+			// If there aren't, then just use the default mapping for this action
+			if (!bIsPlayerMapping)
+			{
+				MappedKeysToActionName.Add(Mapping.Key);
+			}
+			
+			for (const FKey& PlayerMappedKey : MappedKeysToActionName)
 			{
 				Mapping.Key = PlayerMappedKey;
 
diff --git a/Engine/Plugins/EnhancedInput/Source/EnhancedInput/Private/EnhancedInputSubsystems.cpp b/Engine/Plugins/EnhancedInput/Source/EnhancedInput/Private/EnhancedInputSubsystems.cpp
--- a/Engine/Plugins/EnhancedInput/Source/EnhancedInput/Private/EnhancedInputSubsystems.cpp
+++ b/Engine/Plugins/EnhancedInput/Source/EnhancedInput/Private/EnhancedInputSubsystems.cpp
@@ -22,6 +22,18 @@ DEFINE_LOG_CATEGORY(LogWorldSubsystemInput);
 // *
 // **************************************************************************************************
 
+void UEnhancedInputLocalPlayerSubsystem::PlayerControllerChanged(APlayerController* NewPlayerController)
+{
+	Super::PlayerControllerChanged(NewPlayerController);
+
+	// Now that there is a valid player controller, we should have access to a valid PlayerInput pointer.
+	// That means that we can listen for changes to that player's settings!
+	if (GetDefault<UEnhancedInputDeveloperSettings>()->bEnableUserSettings)
+	{
+		BindUserSettingDelegates();	
+	}
+}
+
 UEnhancedPlayerInput* UEnhancedInputLocalPlayerSubsystem::GetPlayerInput() const
 {	
 	if (APlayerController* PlayerController = GetLocalPlayer()->GetPlayerController(GetWorld()))
diff --git a/Engine/Plugins/EnhancedInput/Source/EnhancedInput/Public/EnhancedInputSubsystemInterface.h b/Engine/Plugins/EnhancedInput/Source/EnhancedInput/Public/EnhancedInputSubsystemInterface.h
--- a/Engine/Plugins/EnhancedInput/Source/EnhancedInput/Public/EnhancedInputSubsystemInterface.h
+++ b/Engine/Plugins/EnhancedInput/Source/EnhancedInput/Public/EnhancedInputSubsystemInterface.h
@@ -88,9 +88,11 @@ protected:
 	virtual void BindUserSettingDelegates();
 	
 	/** Callback for when any Enhanced Input user settings have been changed (a new key mapping for example) */
+	UFUNCTION()
 	virtual void OnUserSettingsChanged(UEnhancedInputUserSettings* Settings);
 
 	/** A callback for when the user has applied a new mappable key profile. */
+	UFUNCTION()
 	virtual void OnUserKeyProfileChanged(const UEnhancedPlayerMappableKeyProfile* InNewProfile);
 
 public:
diff --git a/Engine/Plugins/EnhancedInput/Source/EnhancedInput/Public/EnhancedInputSubsystems.h b/Engine/Plugins/EnhancedInput/Source/EnhancedInput/Public/EnhancedInputSubsystems.h
--- a/Engine/Plugins/EnhancedInput/Source/EnhancedInput/Public/EnhancedInputSubsystems.h
+++ b/Engine/Plugins/EnhancedInput/Source/EnhancedInput/Public/EnhancedInputSubsystems.h
@@ -22,6 +22,10 @@ class ENHANCEDINPUT_API UEnhancedInputLocalPlayerSubsystem : public ULocalPlayer
 
 public:
 
+	// Begin ULocalPlayerSubsystem
+	virtual void PlayerControllerChanged(APlayerController* NewPlayerController) override;
+	// End ULocalPlayerSubsystem
+	
 	// Begin IEnhancedInputSubsystemInterface
 	virtual UEnhancedPlayerInput* GetPlayerInput() const override;
 	virtual void ControlMappingsRebuiltThisFrame() override;
