Do not initalize the Config Display name in a Player Mappable Input Config from the name of the action because it will result in improper localization data being generated and you will get cook warnings.

#rnx

diff --git a/Engine/Plugins/EnhancedInput/Source/EnhancedInput/Private/PlayerMappableInputConfig.cpp b/Engine/Plugins/EnhancedInput/Source/EnhancedInput/Private/PlayerMappableInputConfig.cpp
--- a/Engine/Plugins/EnhancedInput/Source/EnhancedInput/Private/PlayerMappableInputConfig.cpp
+++ b/Engine/Plugins/EnhancedInput/Source/EnhancedInput/Private/PlayerMappableInputConfig.cpp
@@ -10,7 +10,7 @@
 UPlayerMappableInputConfig::UPlayerMappableInputConfig(const FObjectInitializer& ObjectInitializer)
 	: Super(ObjectInitializer)
 	, ConfigName(GetFName())
-	, ConfigDisplayName(FText::FromName(ConfigName))
+	, ConfigDisplayName()
 {
 	ResetToDefault();
 }
