Make ResetToDefault a public facing function on the Key Profile and blueprint callable

#rnx
#pf 63e7cd35923476b68635c2d5

diff --git a/Engine/Plugins/EnhancedInput/Source/EnhancedInput/Public/UserSettings/EnhancedInputUserSettings.h b/Engine/Plugins/EnhancedInput/Source/EnhancedInput/Public/UserSettings/EnhancedInputUserSettings.h
--- a/Engine/Plugins/EnhancedInput/Source/EnhancedInput/Public/UserSettings/EnhancedInputUserSettings.h
+++ b/Engine/Plugins/EnhancedInput/Source/EnhancedInput/Public/UserSettings/EnhancedInputUserSettings.h
@@ -261,16 +261,14 @@ public:
 	
 	UFUNCTION(BlueprintCallable, Category="Enhanced Input|User Settings", meta=(DisplayName="Find Key Mapping", AutoCreateRefTerm="OutKeyMapping"))
 	void K2_FindKeyMapping(FPlayerKeyMapping& OutKeyMapping, const FMapPlayerKeyArgs& InArgs) const;
-	
-protected:
 
-	
 	/**
 	 * Resets all the key mappings in this profile to their default value from their Input Mapping Context.
-	 *
-	 * This function will only ever be called by UEnhancedInputUserSettings
 	 */
+	UFUNCTION(BlueprintCallable, Category="Enhanced Input|User Settings")
 	virtual void ResetToDefault();
+	
+protected:
 
 	/**
 	 * Equips the current key profile. This will always be called after the previous key profile's UnEquip function.
