Fix deprecation macro "mismatching"
#rnx

diff --git a/Engine/Plugins/EnhancedInput/Source/EnhancedInput/Public/EnhancedActionKeyMapping.h b/Engine/Plugins/EnhancedInput/Source/EnhancedInput/Public/EnhancedActionKeyMapping.h
--- a/Engine/Plugins/EnhancedInput/Source/EnhancedInput/Public/EnhancedActionKeyMapping.h
+++ b/Engine/Plugins/EnhancedInput/Source/EnhancedInput/Public/EnhancedActionKeyMapping.h
@@ -113,12 +113,10 @@ struct ENHANCEDINPUT_API FEnhancedActionKeyMapping
 
 #if WITH_EDITORONLY_DATA
 
-	PRAGMA_DISABLE_DEPRECATION_WARNINGS
 	/** Options for making this a player mappable keymapping */
 	UE_DEPRECATED(5.3, "PlayerMappableOptions has been deprecated, please use PlayerMappableKeySettings instead.")
 	UPROPERTY(EditAnywhere, Category = "Input|PlayerMappable", meta = (editCondition = "bIsPlayerMappable", DisplayAfter = "bIsPlayerMappable", DeprecatedProperty, DeprecationMessage="PlayerMappableOptions has been deprecated, please use the PlayerMappableKeySettings instead"))
 	FPlayerMappableKeyOptions PlayerMappableOptions;
-	PRAGMA_ENABLE_DEPRECATION_WARNINGS
 
 #endif	// WITH_EDITORONLY_DATA
 	
@@ -156,12 +154,10 @@ struct ENHANCEDINPUT_API FEnhancedActionKeyMapping
 
 #if WITH_EDITORONLY_DATA
 
-	PRAGMA_DISABLE_DEPRECATION_WARNINGS
 	/** If true then this ActionKeyMapping will be exposed as a player mappable key */
 	UE_DEPRECATED(5.3, "bIsPlayerMappable has been deprecated, please use SettingBehavior instead.")
 	UPROPERTY(EditAnywhere, Category = "Input|PlayerMappable", meta=(DeprecatedProperty, DeprecationMessage="bIsPlayerMappable has been deprecated, please use the SettingBehavior instead"))
 	uint8 bIsPlayerMappable : 1;
-	PRAGMA_ENABLE_DEPRECATION_WARNINGS
 
 #endif	// WITH_EDITORONLY_DATA
 
