Fixing ensure when no combo actions are added

diff --git a/Engine/Plugins/EnhancedInput/Source/EnhancedInput/Private/InputTriggers.cpp b/Engine/Plugins/EnhancedInput/Source/EnhancedInput/Private/InputTriggers.cpp
--- a/Engine/Plugins/EnhancedInput/Source/EnhancedInput/Private/InputTriggers.cpp
+++ b/Engine/Plugins/EnhancedInput/Source/EnhancedInput/Private/InputTriggers.cpp
@@ -1,6 +1,7 @@
 // Copyright Epic Games, Inc. All Rights Reserved.
 
 #include "InputTriggers.h"
+#include "EnhancedInputModule.h"
 #include "EnhancedPlayerInput.h"
 
 #include UE_INLINE_GENERATED_CPP_BY_NAME(InputTriggers)
@@ -234,7 +235,7 @@ ETriggerState UInputTriggerCombo::UpdateState_Implementation(const UEnhancedPlay
 {
 	if (ComboActions.IsEmpty())
 	{
-		ensureMsgf(false, TEXT("You must add combo actions to the UInputTriggerCombo for it work properly! Exiting..."));
+		UE_LOG(LogEnhancedInput, Warning, TEXT("A Combo Trigger has no combo actions and will not work properly!"));
 		return ETriggerState::None;
 	}
 	
