Fix the Default profile ID tag not being correct

#pf skip
#rnx

diff --git a/Engine/Plugins/EnhancedInput/Source/EnhancedInput/Private/UserSettings/EnhancedInputUserSettings.cpp b/Engine/Plugins/EnhancedInput/Source/EnhancedInput/Private/UserSettings/EnhancedInputUserSettings.cpp
--- a/Engine/Plugins/EnhancedInput/Source/EnhancedInput/Private/UserSettings/EnhancedInputUserSettings.cpp
+++ b/Engine/Plugins/EnhancedInput/Source/EnhancedInput/Private/UserSettings/EnhancedInputUserSettings.cpp
@@ -27,7 +27,7 @@ namespace UE::EnhancedInput
 	/** The name of the slot that these settings will save to */
 	static const FString SETTINGS_SLOT_NAME = TEXT("EnhancedInputUserSettings");
 
-	UE_DEFINE_GAMEPLAY_TAG_STATIC(TAG_DefaultProfileIdentifier, "InputUserSettings.FailureReasons.InvalidActionName");
+	UE_DEFINE_GAMEPLAY_TAG_STATIC(TAG_DefaultProfileIdentifier, "InputUserSettings.Profiles.Default");
 	static const FText DefaultProfileDisplayName = LOCTEXT("Default_Profile_name", "Default Profile");
 	
 	UE_DEFINE_GAMEPLAY_TAG_STATIC(TAG_InvalidActionName, "InputUserSettings.FailureReasons.InvalidActionName");
