Add a name validator to player mappable setting names. This makes it easier to avoid mistakes and possibly stomping a key binding.

The "Name" here is the FName that will be used as a key to save each key mapping during serialization, so they should be unique for each key mapping that you want the player to be able to set.

#rnx

diff --git a/Engine/Plugins/EnhancedInput/Source/EnhancedInput/Private/PlayerMappableKeySettings.cpp b/Engine/Plugins/EnhancedInput/Source/EnhancedInput/Private/PlayerMappableKeySettings.cpp
--- a/Engine/Plugins/EnhancedInput/Source/EnhancedInput/Private/PlayerMappableKeySettings.cpp
+++ b/Engine/Plugins/EnhancedInput/Source/EnhancedInput/Private/PlayerMappableKeySettings.cpp
@@ -7,6 +7,9 @@
 #define LOCTEXT_NAMESPACE "EnhancedActionKeySetting"
 
 #if WITH_EDITOR
+
+#include "UObject/UObjectIterator.h"
+
 EDataValidationResult UPlayerMappableKeySettings::IsDataValid(TArray<FText>& ValidationErrors)
 {
 	EDataValidationResult Result = CombineDataValidationResults(Super::IsDataValid(ValidationErrors), EDataValidationResult::Valid);
@@ -17,6 +20,23 @@ EDataValidationResult UPlayerMappableKeySettings::IsDataValid(TArray<FText>& Val
 	}
 	return Result;
 }
+
+const TArray<FName>& UPlayerMappableKeySettings::GetKnownMappingNames()
+{
+	static TArray<FName> OutNames;
+	OutNames.Reset();
+	
+    for (TObjectIterator<UPlayerMappableKeySettings> Itr; Itr; ++Itr)
+    {
+    	if (IsValid(*Itr))
+    	{
+    		OutNames.Add(Itr->Name);	
+    	}
+    }
+
+    return OutNames;
+}
+
 #endif // WITH_EDITOR
 
 #undef LOCTEXT_NAMESPACE
diff --git a/Engine/Plugins/EnhancedInput/Source/EnhancedInput/Public/PlayerMappableKeySettings.h b/Engine/Plugins/EnhancedInput/Source/EnhancedInput/Public/PlayerMappableKeySettings.h
--- a/Engine/Plugins/EnhancedInput/Source/EnhancedInput/Public/PlayerMappableKeySettings.h
+++ b/Engine/Plugins/EnhancedInput/Source/EnhancedInput/Public/PlayerMappableKeySettings.h
@@ -22,6 +22,16 @@ public:
 
 #if WITH_EDITOR
 	virtual EDataValidationResult IsDataValid(TArray<FText>& ValidationErrors) override;
+
+	/**
+	 * Get the known mapping names that are current in use. This is a helper function if you want to use a "GetOptions" metadata on a UPROPERTY.
+	 * For example, the following will display a little drop down menu to select from all current mapping names:
+	 *
+	 *  UPROPERTY(EditAnywhere, BlueprintReadOnly, meta=(GetOptions="EnhancedInput.PlayerMappableKeySettings.GetKnownMappingNames"))
+	 *  FName MappingName;
+	 */
+	UFUNCTION()
+	static const TArray<FName>& GetKnownMappingNames();	
 #endif // WITH_EDITOR
 
 	/** Metadata that can used to store any other related items to this key mapping such as icons, ability assets, etc. */
diff --git a/Engine/Plugins/EnhancedInput/Source/InputEditor/Private/EnhancedInputPlayerMappableNameValidator.cpp b/Engine/Plugins/EnhancedInput/Source/InputEditor/Private/EnhancedInputPlayerMappableNameValidator.cpp
new file mode 100644
--- /dev/null
+++ b/Engine/Plugins/EnhancedInput/Source/InputEditor/Private/EnhancedInputPlayerMappableNameValidator.cpp
@@ -0,0 +1,20 @@
+// Copyright Epic Games, Inc. All Rights Reserved.
+
+#include "EnhancedInputPlayerMappableNameValidator.h"
+#include "InputEditorModule.h"
+
+FEnhancedInputPlayerMappableNameValidator::FEnhancedInputPlayerMappableNameValidator(FName InExistingName)
+	: FStringSetNameValidator(InExistingName.ToString())
+{ }
+
+EValidatorResult FEnhancedInputPlayerMappableNameValidator::IsValid(const FString& Name, bool bOriginal)
+{
+	EValidatorResult Result = FStringSetNameValidator::IsValid(Name, bOriginal);
+
+	if (Result != EValidatorResult::ExistingName && FInputEditorModule::IsMappingNameInUse(FName(Name)))
+	{
+		Result = EValidatorResult::AlreadyInUse;
+	}
+
+	return Result;
+}
\ No newline at end of file
diff --git a/Engine/Plugins/EnhancedInput/Source/InputEditor/Private/InputCustomizations.cpp b/Engine/Plugins/EnhancedInput/Source/InputEditor/Private/InputCustomizations.cpp
--- a/Engine/Plugins/EnhancedInput/Source/InputEditor/Private/InputCustomizations.cpp
+++ b/Engine/Plugins/EnhancedInput/Source/InputEditor/Private/InputCustomizations.cpp
@@ -17,6 +17,11 @@
 #include "UObject/UObjectIterator.h"
 #include "Modules/ModuleManager.h"
 #include "Widgets/Text/STextBlock.h"
+#include "EnhancedInputPlayerMappableNameValidator.h"
+#include "PlayerMappableKeySettings.h"
+#include "Widgets/Input/SEditableTextBox.h"
+#include "EnhancedInputPlayerMappableNameValidator.h"
+#include "InputEditorModule.h"
 
 #define LOCTEXT_NAMESPACE "InputCustomization"
 
@@ -489,4 +494,174 @@ void FEnhancedInputDeveloperSettingsCustomization::RebuildDetailsViewForAsset(co
 	}
 }
 
-#undef LOCTEXT_NAMESPACE
+/////////////////////////////////////////////////////////////////
+// FPlayerMappableKeyChildSettingsCustomization
+
+namespace UE::EnhancedInput
+{
+	/** Flag to enable or disable name validation in the editor for UPlayerMappableKeySettings assets */
+	static bool bEnableNameValidation = true;
+	static FAutoConsoleVariableRef CVarEnableNameValidation(
+		TEXT("EnhancedInput.bEnableNameValidation"),
+		bEnableNameValidation,
+		TEXT("Flag to enable or disable name validation in the editor for UPlayerMappableKeySettings assets"),
+		ECVF_Default);
+}
+
+TSharedRef<IPropertyTypeCustomization> FPlayerMappableKeyChildSettingsCustomization::MakeInstance()
+{
+	return MakeShareable(new FPlayerMappableKeyChildSettingsCustomization());
+}
+
+void FPlayerMappableKeyChildSettingsCustomization::CustomizeHeader(TSharedRef<IPropertyHandle> PropertyHandle, FDetailWidgetRow& HeaderRow, IPropertyTypeCustomizationUtils& CustomizationUtils)
+{
+	SettingsPropHandle = PropertyHandle;
+	UPlayerMappableKeySettings* Settings = GetSettingsObject();
+	
+	HeaderRow.NameContent()
+	[
+		PropertyHandle->CreatePropertyNameWidget()
+	];
+	
+	HeaderRow.ValueContent()
+	[
+		PropertyHandle->CreatePropertyValueWidget()
+	];	
+}
+
+void FPlayerMappableKeyChildSettingsCustomization::CustomizeChildren(TSharedRef<IPropertyHandle> PropertyHandle, IDetailChildrenBuilder& ChildBuilder, IPropertyTypeCustomizationUtils& CustomizationUtils)
+{				
+	// Customize the name handle to have some validation on what users can type in
+	static const FName MappingPropertyHandleName = GET_MEMBER_NAME_CHECKED(UPlayerMappableKeySettings, Name);
+
+	uint32 NumChildren = 0;
+	TSharedPtr<IPropertyHandle> SettingsHandle;
+	TSharedPtr<IPropertyHandle> TempHandle = PropertyHandle->GetChildHandle(0);
+
+	if (TempHandle.IsValid())
+	{
+		TempHandle->GetNumChildren(NumChildren);
+
+		// For instanced properties there can be multiple layers of "children" until we get to the ones
+		// that we care about. Iterate all the children until we find one with more then one
+		// child, or we run out of children.
+		while (!SettingsHandle.IsValid() || !TempHandle.IsValid())
+		{
+			TempHandle->GetNumChildren(NumChildren);
+
+			if (NumChildren > 1)
+			{
+				SettingsHandle = TempHandle;
+			}
+			else
+			{
+				TempHandle = TempHandle->GetChildHandle(0);
+			}		
+		}	
+	}	
+	
+	if (!SettingsHandle.IsValid())
+	{
+		return;
+	}
+
+	for (uint32 ChildIndex = 0; ChildIndex < NumChildren; ++ChildIndex)
+	{
+		TSharedPtr<IPropertyHandle> ChildHandle = SettingsHandle->GetChildHandle(ChildIndex);
+		if (!ChildHandle.IsValid())
+		{
+			continue;
+		}
+		
+		const FName PropertyName = ChildHandle->GetProperty()->GetFName();
+
+		// Add a customization to handle name validation on the mappable name if it's enabled
+		if (PropertyName == MappingPropertyHandleName && UE::EnhancedInput::bEnableNameValidation)
+		{
+			MappingNamePropHandle = ChildHandle;
+			IDetailPropertyRow& NameRow = ChildBuilder.AddProperty(MappingNamePropHandle.ToSharedRef());       
+            
+            NameRow.CustomWidget()
+            .NameContent()
+            [
+            	MappingNamePropHandle->CreatePropertyNameWidget()
+            ]
+            .ValueContent()
+            [
+            	SNew(SEditableTextBox)
+            		.Text(this, &FPlayerMappableKeyChildSettingsCustomization::OnGetMappingNameText)
+            		.OnTextCommitted(this, &FPlayerMappableKeyChildSettingsCustomization::OnMappingNameTextCommited)
+            		.ToolTipText(this, &FPlayerMappableKeyChildSettingsCustomization::OnGetMappingNameToolTipText)
+            		.OnVerifyTextChanged(this, &FPlayerMappableKeyChildSettingsCustomization::OnVerifyMappingName)
+            		.Font(IDetailLayoutBuilder::GetDetailFont())
+            ];
+		}
+		else
+		{
+			ChildBuilder.AddProperty(ChildHandle.ToSharedRef());
+		}
+	}
+}
+
+FName FPlayerMappableKeyChildSettingsCustomization::GetCurrentName() const
+{
+	if (UPlayerMappableKeySettings* Settings = GetSettingsObject())
+	{
+		return Settings->Name;
+	}
+	ensure(false);
+	return NAME_None;
+}
+
+UPlayerMappableKeySettings* FPlayerMappableKeyChildSettingsCustomization::GetSettingsObject() const
+{
+	if (SettingsPropHandle.IsValid())
+	{
+		UObject* SettingsObject = nullptr;
+		SettingsPropHandle->GetValue(SettingsObject);
+		return Cast<UPlayerMappableKeySettings>(SettingsObject);	
+	}
+
+	return nullptr;
+}
+
+FText FPlayerMappableKeyChildSettingsCustomization::OnGetMappingNameText() const
+{
+	// this is the value in the actual text box
+	if (UPlayerMappableKeySettings* Settings = GetSettingsObject())
+	{
+		return FText::FromName(Settings->Name);
+	}
+	
+	return FText::FromName(NAME_None);
+}
+
+FText FPlayerMappableKeyChildSettingsCustomization::OnGetMappingNameToolTipText() const
+{
+	static const FText Tooltip = LOCTEXT("MappingNameToolTip", "This mapping name will be the key that is used to save player mapping data for this action. It can be overriden by individual mappings, or specified by the Input Action. These names need to be unique for each player mapping you wish to save.");
+	return Tooltip;
+}
+
+void FPlayerMappableKeyChildSettingsCustomization::OnMappingNameTextCommited(const FText& NewText, ETextCommit::Type InTextCommit)
+{
+	// set the actual property value on the settings object
+	if (InTextCommit == ETextCommit::OnEnter || InTextCommit == ETextCommit::OnUserMovedFocus)
+	{
+		MappingNamePropHandle->SetValue(NewText.ToString());
+	}
+}
+
+bool FPlayerMappableKeyChildSettingsCustomization::OnVerifyMappingName(const FText& InNewText, FText& OutErrorMessage)
+{
+	UPlayerMappableKeySettings* Settings = GetSettingsObject();
+	const FName CurrentName = Settings ? Settings->Name : NAME_None;
+	
+	// Validate that the new name can be used
+	FEnhancedInputPlayerMappableNameValidator NameValidator(CurrentName);
+	EValidatorResult Result = NameValidator.IsValid(InNewText.ToString(), false);
+	OutErrorMessage = INameValidatorInterface::GetErrorText(InNewText.ToString(), Result);
+
+	return Result == EValidatorResult::Ok || Result == EValidatorResult::ExistingName;
+}
+
+#undef LOCTEXT_NAMESPACE
\ No newline at end of file
diff --git a/Engine/Plugins/EnhancedInput/Source/InputEditor/Private/InputCustomizations.h b/Engine/Plugins/EnhancedInput/Source/InputEditor/Private/InputCustomizations.h
--- a/Engine/Plugins/EnhancedInput/Source/InputEditor/Private/InputCustomizations.h
+++ b/Engine/Plugins/EnhancedInput/Source/InputEditor/Private/InputCustomizations.h
@@ -6,7 +6,7 @@
 #include "IPropertyTypeCustomization.h"
 
 class FKeyStructCustomization;
-
+class UPlayerMappableKeySettings;
 class IDetailPropertyRow;
 class IDetailLayoutBuilder;
 
@@ -109,3 +109,30 @@ private:
 	 */
 	static TSet<FName> ExcludedAssetNames;
 };
+
+/**
+ * Input customization for UPlayerMappableKeySettings to have some name validation
+ * on the custom saved FNames for mappings.
+ */
+class FPlayerMappableKeyChildSettingsCustomization : public IPropertyTypeCustomization
+{
+public:
+	static TSharedRef<IPropertyTypeCustomization> MakeInstance();
+
+	/** IPropertyTypeCustomization interface */
+	virtual void CustomizeHeader(TSharedRef<IPropertyHandle> PropertyHandle, FDetailWidgetRow& HeaderRow, IPropertyTypeCustomizationUtils& CustomizationUtils) override;
+	virtual void CustomizeChildren(TSharedRef<IPropertyHandle> PropertyHandle, IDetailChildrenBuilder& ChildBuilder, IPropertyTypeCustomizationUtils& CustomizationUtils) override;
+
+private:
+	
+	UPlayerMappableKeySettings* GetSettingsObject() const;
+	FName GetCurrentName() const;
+	
+	FText OnGetMappingNameText() const;
+	FText OnGetMappingNameToolTipText() const;
+	void OnMappingNameTextCommited(const FText& NewText, ETextCommit::Type InTextCommit);
+	bool OnVerifyMappingName(const FText& InNewText, FText& OutErrorMessage);
+
+	TSharedPtr<IPropertyHandle> SettingsPropHandle;
+	TSharedPtr<IPropertyHandle> MappingNamePropHandle;
+};
\ No newline at end of file
diff --git a/Engine/Plugins/EnhancedInput/Source/InputEditor/Private/InputEditorModule.cpp b/Engine/Plugins/EnhancedInput/Source/InputEditor/Private/InputEditorModule.cpp
--- a/Engine/Plugins/EnhancedInput/Source/InputEditor/Private/InputEditorModule.cpp
+++ b/Engine/Plugins/EnhancedInput/Source/InputEditor/Private/InputEditorModule.cpp
@@ -39,6 +39,7 @@
 #include "SourceControlHelpers.h"
 #include "HAL/FileManager.h"
 #include "Kismet2/KismetEditorUtilities.h"
+#include "PlayerMappableKeySettings.h"
 
 #include UE_INLINE_GENERATED_CPP_BY_NAME(InputEditorModule)
 
@@ -409,6 +410,7 @@ void FInputEditorModule::StartupModule()
 	PropertyModule.RegisterCustomClassLayout("InputMappingContext", FOnGetDetailCustomizationInstance::CreateStatic(&FInputContextDetails::MakeInstance));
 	PropertyModule.RegisterCustomPropertyTypeLayout("EnhancedActionKeyMapping", FOnGetPropertyTypeCustomizationInstance::CreateStatic(&FEnhancedActionMappingCustomization::MakeInstance));
 	PropertyModule.RegisterCustomClassLayout(UEnhancedInputDeveloperSettings::StaticClass()->GetFName(), FOnGetDetailCustomizationInstance::CreateStatic(&FEnhancedInputDeveloperSettingsCustomization::MakeInstance));
+	PropertyModule.RegisterCustomPropertyTypeLayout(UPlayerMappableKeySettings::StaticClass()->GetFName(), FOnGetPropertyTypeCustomizationInstance::CreateStatic(&FPlayerMappableKeyChildSettingsCustomization::MakeInstance));
 	PropertyModule.NotifyCustomizationModuleChanged();
 
 	// Register input assets
@@ -462,6 +464,7 @@ void FInputEditorModule::ShutdownModule()
 	PropertyModule.UnregisterCustomClassLayout("InputContext");
 	PropertyModule.UnregisterCustomPropertyTypeLayout("EnhancedActionKeyMapping");
 	PropertyModule.UnregisterCustomClassLayout("EnhancedInputDeveloperSettings");
+	PropertyModule.UnregisterCustomPropertyTypeLayout("PlayerMappableKeySettings");
 	PropertyModule.NotifyCustomizationModuleChanged();
 
 	// Unregister slate stylings
@@ -552,6 +555,21 @@ void FInputEditorModule::AutoUpgradeDefaultInputClasses()
 	}
 }
 
+bool FInputEditorModule::IsMappingNameInUse(const FName InName)
+{
+	// Instanced objects that have been deleted will have an outer of the transient package, and we don't want to include them
+	for (TObjectIterator<UPlayerMappableKeySettings> Itr(RF_ClassDefaultObject | RF_Transient, true, EInternalObjectFlags::Garbage | EInternalObjectFlags::Unreachable); Itr; ++Itr)
+	{
+		UPlayerMappableKeySettings* Settings = *Itr;		
+		if (Settings && Settings->GetOuter() != GetTransientPackage() && Settings->Name == InName)
+		{
+			return true;
+		}
+	}
+	
+	return false;
+}
+
 void FInputEditorModule::Tick(float DeltaTime)
 {
 	// Update any blueprints that are referencing an input action with a modified value type
@@ -611,4 +629,4 @@ void FInputEditorModule::Tick(float DeltaTime)
 	}
 }
 
-#undef LOCTEXT_NAMESPACE
+#undef LOCTEXT_NAMESPACE
\ No newline at end of file
diff --git a/Engine/Plugins/EnhancedInput/Source/InputEditor/Public/EnhancedInputPlayerMappableNameValidator.h b/Engine/Plugins/EnhancedInput/Source/InputEditor/Public/EnhancedInputPlayerMappableNameValidator.h
new file mode 100644
--- /dev/null
+++ b/Engine/Plugins/EnhancedInput/Source/InputEditor/Public/EnhancedInputPlayerMappableNameValidator.h
@@ -0,0 +1,15 @@
+// Copyright Epic Games, Inc. All Rights Reserved.
+
+#pragma once
+
+#include "Kismet2/Kismet2NameValidators.h"
+
+class FEnhancedInputPlayerMappableNameValidator : public FStringSetNameValidator
+{
+public:
+	FEnhancedInputPlayerMappableNameValidator(FName InExistingName);
+
+	// Begin FNameValidatorInterface
+	virtual EValidatorResult IsValid(const FString& Name, bool bOriginal) override;
+	// End FNameValidatorInterface
+};
\ No newline at end of file
diff --git a/Engine/Plugins/EnhancedInput/Source/InputEditor/Public/InputEditorModule.h b/Engine/Plugins/EnhancedInput/Source/InputEditor/Public/InputEditorModule.h
--- a/Engine/Plugins/EnhancedInput/Source/InputEditor/Public/InputEditorModule.h
+++ b/Engine/Plugins/EnhancedInput/Source/InputEditor/Public/InputEditorModule.h
@@ -33,6 +33,9 @@ public:
 
 	static EAssetTypeCategories::Type GetInputAssetsCategory() { return InputAssetsCategory; }
 	
+	/** Returns true if the given name is in use by a player mappable key setting */
+	INPUTEDITOR_API static bool IsMappingNameInUse(const FName InName);
+		
 private:
 	void RegisterAssetTypeActions(IAssetTools& AssetTools, TSharedRef<IAssetTypeActions> Action)
 	{
@@ -96,23 +99,6 @@ public:
 	virtual UObject* FactoryCreateNew(UClass* Class, UObject* InParent, FName Name, EObjectFlags Flags, UObject* Context, FFeedbackContext* Warn) override;
 };
 
-// TODO: Add trigger/modifier factories and hook up RegisterAssetTypeActions type construction.
-//
-//UCLASS()
-//class INPUTEDITOR_API UInputTrigger_Factory : public UBlueprintFactory
-//{
-//	GENERATED_UCLASS_BODY()
-//};
-//
-//UCLASS()
-//class INPUTEDITOR_API UInputModifier_Factory : public UBlueprintFactory
-//{
-//	GENERATED_UCLASS_BODY()
-//
-//	UPROPERTY(EditAnywhere, Category = DataAsset)
-//	TSubclassOf<UDataAsset> DataAssetClass;
-//};
-
 #if UE_ENABLE_INCLUDE_ORDER_DEPRECATED_IN_5_2
 #include "CoreMinimal.h"
 #include "IDetailsView.h"
