Only serialize key mappings that the user has changed when saving the player key profile. This will signifigantly reduce the amount of data we have to save (~10x in Lyra, going from a 63kb save file to a sub 6kb save file) .

We don't store a copy fo the original key mapping for player mappable keys. Instead, we can populate the metadata for that key mapping when a mapping context is registered.

Add a comment on the "SaveSettings" function that mentions you can also save this data by simply calling the Serialize function, which is what more complex situations will require such as cloud save.

Mark several other UPROPERTY's as transient so that if you do use the SaveGame system in a local game without cloud save, you dont have unnecessary data being written out.

#rnx

diff --git a/Engine/Plugins/EnhancedInput/Source/EnhancedInput/Private/UserSettings/EnhancedInputUserSettings.cpp b/Engine/Plugins/EnhancedInput/Source/EnhancedInput/Private/UserSettings/EnhancedInputUserSettings.cpp
--- a/Engine/Plugins/EnhancedInput/Source/EnhancedInput/Private/UserSettings/EnhancedInputUserSettings.cpp
+++ b/Engine/Plugins/EnhancedInput/Source/EnhancedInput/Private/UserSettings/EnhancedInputUserSettings.cpp
@@ -93,6 +93,16 @@ FPlayerKeyMapping::FPlayerKeyMapping()
 {
 }
 
+FPlayerKeyMapping::FPlayerKeyMapping(const FEnhancedActionKeyMapping& OriginalMapping, EPlayerMappableKeySlot InSlot /* = EPlayerMappableKeySlot::Unspecified */)
+	: ActionName(OriginalMapping.GetMappingName())
+	, DisplayName(OriginalMapping.GetDisplayName())
+	, Slot(InSlot)
+	, DefaultKey(OriginalMapping.Key)
+	, CurrentKey(EKeys::Invalid)
+	, HardwareDeviceId(FHardwareDeviceIdentifier::Invalid)
+{
+}
+
 // The default constructor creates an invalid mapping. Use this as a way to return references
 // to an invalid mapping for BP functions
 FPlayerKeyMapping FPlayerKeyMapping::InvalidMapping = FPlayerKeyMapping();
@@ -137,14 +147,6 @@ const FName FPlayerKeyMapping::GetActionName() const
 
 const FText& FPlayerKeyMapping::GetDisplayName() const
 {
-	// Just in case the display name is empty on this mapping, see if we can fall back to the original mapping copy's display name
-	if (DisplayName.IsEmpty())
-	{
-		if (const UPlayerMappableKeySettings* Settings = OriginalMappingCopy.GetPlayerMappableKeySettings())
-		{
-			return Settings->DisplayName;
-		}
-	}
 	return DisplayName;
 }
 
@@ -170,12 +172,35 @@ uint32 GetTypeHash(const FPlayerKeyMapping& InMapping)
 
 void FPlayerKeyMapping::ResetToDefault()
 {
-	CurrentKey = DefaultKey;
+	// Only mark as dirty during a reset to default if the mapping was set to a different key
+	// This avoids unnecessarily marking keys as dirty if you do a reset of the entire key profile
+	if (IsCustomized())
+	{
+		bIsDirty = true;
+	}
+	
+	CurrentKey = DefaultKey;	
 }
 
 void FPlayerKeyMapping::SetCurrentKey(const FKey& NewKey)
 {
 	CurrentKey = NewKey;
+	bIsDirty = true;
+}
+
+void FPlayerKeyMapping::UpdateOriginalKey(const FEnhancedActionKeyMapping& OriginalMapping)
+{
+	ensure(OriginalMapping.IsPlayerMappable() && OriginalMapping.GetMappingName() == ActionName);
+	
+	DefaultKey = OriginalMapping.Key;
+	DisplayName = OriginalMapping.GetDisplayName();
+
+	// If the mapping is the same as the default key, make sure that it is not marked as
+	// dirty so that we don't serialize it
+	if (!IsCustomized())
+	{
+		bIsDirty = false;
+	}
 }
 
 bool FPlayerKeyMapping::operator==(const FPlayerKeyMapping& Other) const
@@ -184,8 +209,8 @@ bool FPlayerKeyMapping::operator==(const FPlayerKeyMapping& Other) const
 	       ActionName			== Other.ActionName
 		&& Slot					== Other.Slot
 		&& HardwareDeviceId		== Other.HardwareDeviceId
-		&& CurrentKey		== Other.CurrentKey
-		&& OriginalMappingCopy	== Other.OriginalMappingCopy;
+		&& CurrentKey			== Other.CurrentKey
+		&& DefaultKey			== Other.DefaultKey;
 }
 
 bool FPlayerKeyMapping::operator!=(const FPlayerKeyMapping& Other) const
@@ -193,6 +218,11 @@ bool FPlayerKeyMapping::operator!=(const FPlayerKeyMapping& Other) const
 	return !FPlayerKeyMapping::operator==(Other);
 }
 
+const bool FPlayerKeyMapping::IsDirty() const
+{
+	return bIsDirty;
+}
+
 ///////////////////////////////////////////////////////////
 // UEnhancedPlayerMappableKeyProfile
 
@@ -438,6 +468,23 @@ namespace UE::EnhancedInput
 {
 	static const int32 GPlayerMappableSaveVersion = 1;
 
+	struct FKeyMappingSaveData
+	{
+		friend FArchive& operator<<(FArchive& Ar, FKeyMappingSaveData& Data)
+		{
+			Ar << Data.ActionName;
+			Ar << Data.CurrentKeyName;
+			Ar << Data.HardwareDeviceId;
+			Ar << Data.Slot;
+			return Ar;
+		}
+		
+		FName ActionName = NAME_None;
+		FName CurrentKeyName = NAME_None;
+		FHardwareDeviceIdentifier HardwareDeviceId = FHardwareDeviceIdentifier::Invalid;
+		EPlayerMappableKeySlot Slot = EPlayerMappableKeySlot::Unspecified;
+	};
+
 	/** Struct used to store info about the mappable profile subobjects */
 	struct FMappableKeysHeader
 	{
@@ -446,12 +493,14 @@ namespace UE::EnhancedInput
 			Ar << Header.ProfileIdentifier;
 			Ar << Header.ClassPath;
 			Ar << Header.ObjectPath;
+			Ar << Header.DirtyMappings;
 			return Ar;
 		}
 
 		FGameplayTag ProfileIdentifier;
 		FString ClassPath;
 		FString ObjectPath;
+		TArray<FKeyMappingSaveData> DirtyMappings;
 	};
 }
 
@@ -477,7 +526,33 @@ void UEnhancedInputUserSettings::Serialize(FArchive& Ar)
 		
 		for (TPair<FGameplayTag, UEnhancedPlayerMappableKeyProfile*> ProfilePair : SavedKeyProfiles)
 		{
-			Headers.Push({ ProfilePair.Key, ProfilePair.Value->GetClass()->GetPathName(), ProfilePair.Value->GetPathName(Outer) });
+			UE::EnhancedInput::FMappableKeysHeader Header =
+				{
+					/* .ProfileIdentifier = */ ProfilePair.Key,
+					/* .ClassPath = */ ProfilePair.Value->GetClass()->GetPathName(),
+					/* .ObjectPath = */ ProfilePair.Value->GetPathName(Outer),
+					/* .DirtyMappings = */ {}
+				};
+
+			// Serialize only dirty mappings here
+			for (TPair<FName, FKeyMappingRow>& MappingRow : ProfilePair.Value->PlayerMappedKeys)
+			{
+				for (FPlayerKeyMapping& Mapping : MappingRow.Value.Mappings)
+				{
+					if (Mapping.IsDirty())
+					{
+						Header.DirtyMappings.Push(
+							{
+								/* .ActionName = */ Mapping.ActionName,
+								/* .CurrentKeyName = */ Mapping.CurrentKey.GetFName(),
+								/* .HardwareDeviceId = */ Mapping.HardwareDeviceId, 
+								/* .Slot = */ Mapping.Slot
+							});
+					}
+				}
+			}
+
+			Headers.Push(Header);
 		}
 	}
 	
@@ -491,6 +566,20 @@ void UEnhancedInputUserSettings::Serialize(FArchive& Ar)
 			{
 				UEnhancedPlayerMappableKeyProfile* NewProfile = NewObject<UEnhancedPlayerMappableKeyProfile>(/* outer */ this, /* class */ FoundClass);
 				SavedKeyProfiles.Add(Header.ProfileIdentifier, NewProfile);
+
+				// Add any player saved keys to this profile!
+				for (UE::EnhancedInput::FKeyMappingSaveData& SavedKeyData : Header.DirtyMappings)
+				{
+					FKeyMappingRow& MappingRow = NewProfile->PlayerMappedKeys.FindOrAdd(SavedKeyData.ActionName);
+
+					FPlayerKeyMapping PlayerMapping = {};
+					PlayerMapping.ActionName = SavedKeyData.ActionName;
+					PlayerMapping.CurrentKey = FKey(SavedKeyData.CurrentKeyName);
+					PlayerMapping.Slot = SavedKeyData.Slot;
+					PlayerMapping.HardwareDeviceId = SavedKeyData.HardwareDeviceId;
+					
+					MappingRow.Mappings.Add(PlayerMapping);
+				}
 			}
 		}
 	}
@@ -789,9 +878,10 @@ bool UEnhancedInputUserSettings::RegisterInputMappingContext(UInputMappingContex
 		return false;
 	}
 
-	// There is no need to register an IMC if it already is
-	if( RegisteredMappingContexts.Contains(IMC))
+	// There is no need to re-register an IMC
+	if (RegisteredMappingContexts.Contains(IMC))
 	{
+		UE_LOG(LogEnhancedInput, VeryVerbose, TEXT("Mapping Context '%s' is already registered with the User Settings."), *IMC->GetFName().ToString());
 		return false;
 	}
 	
@@ -804,6 +894,11 @@ bool UEnhancedInputUserSettings::RegisterInputMappingContext(UInputMappingContex
 		ensure(false);
 		return false;
 	}
+
+	// A map that is used to store how many actions of a certain name are in this IMC
+	// this is how we can determine the slot that a new player key mapping needs
+	static TMap<FName, int32> OriginalMappingsDesiredSlot;
+	OriginalMappingsDesiredSlot.Reset();
 	
 	for (const FEnhancedActionKeyMapping& KeyMapping : IMC->GetMappings())
 	{
@@ -818,47 +913,52 @@ bool UEnhancedInputUserSettings::RegisterInputMappingContext(UInputMappingContex
 		// mapping options for a single Input Action. 
 		const FName ActionName = KeyMapping.GetMappingName();
 
-		// If a mapping row exists for this action name, then check if
-		bool bMappingIsInitalized = false;
-		if (FKeyMappingRow* ExistingMappings = CurrentProfile->PlayerMappedKeys.Find(ActionName))
+		// Iterate the number of mappings that use this action
+		int32& NumDefinedMappings = OriginalMappingsDesiredSlot.FindOrAdd(ActionName);
+
+		// Find or create a mapping row 
+		FKeyMappingRow& MappingRow = CurrentProfile->PlayerMappedKeys.FindOrAdd(ActionName);
+		
+		// By default, the slot will be determined by how many mappings this action has already.
+		// So if this is the first default mapping, then this will be EPlayerMappableKeySlot::First,
+		// if this is the second element going into this set then it will be EPlayerMappableKeySlot::Second
+		// and so on
+		const EPlayerMappableKeySlot IMCDefinedSlot =
+			(EPlayerMappableKeySlot)
+			(FMath::Min<uint8>(static_cast<uint8>(NumDefinedMappings), static_cast<uint8>(EPlayerMappableKeySlot::Max)));
+
+		// Update the number of mappings that have this action name
+		++NumDefinedMappings;
+		
+		bool bUpdatedExistingMapping = false;
+		
+		// If a key mapping exists when a mapping context is being registered, that means that it is a 
+		// customized player mapped key.
+		for (FPlayerKeyMapping& ExistingMapping : MappingRow.Mappings)
 		{
-			// Check if we actually have to initalize a new mapping or not
-			for (FPlayerKeyMapping& ExistingMapping : ExistingMappings->Mappings)
+			// We only want to update exiting mapping in the desired slot if one existed in the mapping context.
+			if (ExistingMapping.GetSlot() == IMCDefinedSlot)
 			{
-				// If the original mapping ha already been set, then we know that it's initalized already.
-				// If we don't do this, then registering an IMC will overrwrite whatever saved settings the user has made
-				if (ExistingMapping.OriginalMappingCopy == KeyMapping)
-				{
-					bMappingIsInitalized = true;
-					break;
-				}
+				// Update the fields on this mapping that the player can't change so that they
+				// get updated if the source UInputMappingContext asset has been changed
+				ExistingMapping.UpdateOriginalKey(KeyMapping);
+				bUpdatedExistingMapping = true;
 			}
+			
+			// it is possible that a player has mapped to more slots then you have specified
+			// in a mapping context asset. In that case, they won't have display names populated
+			// and their default key is correctly set to EKeys::Invalid
+			// That is why we need to move the display name to the mapping ROW, since the actions
+			// will have the same display name no matter what.
+			ExistingMapping.DisplayName = KeyMapping.GetDisplayName();
 		}
-		
-		if (!bMappingIsInitalized)
-		{
-			FKeyMappingRow& MappingRow = CurrentProfile->PlayerMappedKeys.Add(ActionName);		
 
+		// If the mapping was not found in the existing row, then the player has not mapped anything to it.
+		// We need to create it based on the original key mapping
+		if (!bUpdatedExistingMapping)
+		{
 			// Add a default mapping to this row
-			FPlayerKeyMapping PlayerMappingData = {};
-			PlayerMappingData.ActionName = ActionName;
-			PlayerMappingData.DefaultKey = KeyMapping.Key;
-			
-			if (UPlayerMappableKeySettings* Settings = KeyMapping.GetPlayerMappableKeySettings())
-			{
-				PlayerMappingData.DisplayName = Settings->DisplayName;
-			}
-			
-			PlayerMappingData.OriginalMappingCopy = KeyMapping;
-
-			// By default, the slot will be determined by how many mappings this action has already.
-			// So if this is the first default mapping, then this will be EPlayerMappableKeySlot::First,
-			// if this is the second element going into this set then it will be EPlayerMappableKeySlot::Second
-			// and so on
-			const uint8 DesiredSlot = FMath::Min<uint8>(static_cast<uint8>(MappingRow.Mappings.Num()), static_cast<uint8>(EPlayerMappableKeySlot::Max));
-			PlayerMappingData.Slot = (EPlayerMappableKeySlot) DesiredSlot;
-		
-			MappingRow.Mappings.Add(PlayerMappingData);
+			MappingRow.Mappings.Add({ KeyMapping, IMCDefinedSlot });	
 		}
 	}
 
diff --git a/Engine/Plugins/EnhancedInput/Source/EnhancedInput/Public/UserSettings/EnhancedInputUserSettings.h b/Engine/Plugins/EnhancedInput/Source/EnhancedInput/Public/UserSettings/EnhancedInputUserSettings.h
--- a/Engine/Plugins/EnhancedInput/Source/EnhancedInput/Public/UserSettings/EnhancedInputUserSettings.h
+++ b/Engine/Plugins/EnhancedInput/Source/EnhancedInput/Public/UserSettings/EnhancedInputUserSettings.h
@@ -87,6 +87,8 @@ struct ENHANCEDINPUT_API FPlayerKeyMapping
 public:
 	
 	FPlayerKeyMapping();
+	
+	FPlayerKeyMapping(const FEnhancedActionKeyMapping& OriginalMapping, EPlayerMappableKeySlot InSlot = EPlayerMappableKeySlot::Unspecified);
 
 	/** A static invalid player key mapping to be used for easy comparisons in blueprint */
 	static FPlayerKeyMapping InvalidMapping;
@@ -129,40 +131,50 @@ public:
 
 	/** Sets the value of the current key to the one given */
 	void SetCurrentKey(const FKey& NewKey);
+
+	/**
+	 * Updates the metadata properties on this player mapped key based on the given
+	 * enhanced action mapping. This will populate the fields on this struct that are not editable
+	 * by the player such as the localized display name and default key.  
+	 */
+	void UpdateOriginalKey(const FEnhancedActionKeyMapping& OriginalMapping);
 	
 	ENHANCEDINPUT_API friend uint32 GetTypeHash(const FPlayerKeyMapping& InMapping);
 	bool operator==(const FPlayerKeyMapping& Other) const;
 	bool operator!=(const FPlayerKeyMapping& Other) const;
 
+	/** Returns true if this mapping has been modified since it was registered from an IMC */
+	const bool IsDirty() const;
+
 protected:
 	
 	/** The name of the action for this key */
-	UPROPERTY(VisibleAnywhere, SaveGame, BlueprintReadOnly, Category="Enhanced Input|User Settings")
+	UPROPERTY(VisibleAnywhere, BlueprintReadOnly, Category="Enhanced Input|User Settings")
 	FName ActionName;
-
+	
 	/** Localized display name of this action */
-	UPROPERTY(VisibleAnywhere, SaveGame, BlueprintReadOnly, Category="Enhanced Input|User Settings")
+	UPROPERTY(Transient, VisibleAnywhere, BlueprintReadOnly, Category="Enhanced Input|User Settings")
 	FText DisplayName;
 
 	/** What slot this key is mapped to */
-	UPROPERTY(VisibleAnywhere, SaveGame, BlueprintReadOnly, Category="Enhanced Input|User Settings")
+	UPROPERTY(VisibleAnywhere, BlueprintReadOnly, Category="Enhanced Input|User Settings")
 	EPlayerMappableKeySlot Slot;
 	
+	/** True if this key mapping is dirty (i.e. has been changed by the player) */
+	UPROPERTY(Transient, VisibleAnywhere, BlueprintReadOnly, Category="Enhanced Input|User Settings")
+	uint8 bIsDirty : 1;
+	
 	/** The default key that this mapping was set to in its input mapping context */
-	UPROPERTY(VisibleAnywhere, SaveGame, BlueprintReadOnly, Category="Enhanced Input|User Settings", meta = (AllowPrivateAccess = "True"))
+	UPROPERTY(Transient, VisibleAnywhere, BlueprintReadOnly, Category="Enhanced Input|User Settings", meta = (AllowPrivateAccess = "True"))
 	FKey DefaultKey;
 	
 	/** The key that the player has mapped this action to */
-	UPROPERTY(VisibleAnywhere, SaveGame,  BlueprintReadOnly, Category="Enhanced Input|User Settings", meta = (AllowPrivateAccess = "True"))
+	UPROPERTY(VisibleAnywhere,  BlueprintReadOnly, Category="Enhanced Input|User Settings", meta = (AllowPrivateAccess = "True"))
 	FKey CurrentKey;
 
 	/** An optional Hardware Device specifier for this mapping */
-	UPROPERTY(VisibleAnywhere, SaveGame, BlueprintReadOnly, Category="Enhanced Input|User Settings")
+	UPROPERTY(VisibleAnywhere, BlueprintReadOnly, Category="Enhanced Input|User Settings")
 	FHardwareDeviceIdentifier HardwareDeviceId;
-
-	/** A copy of the action key mapping that this was generated from. This will only be set on mappings that were created by a Registered IMC */
-	UPROPERTY(VisibleAnywhere, SaveGame, BlueprintReadOnly, Category="Enhanced Input|User Settings")
-	FEnhancedActionKeyMapping OriginalMappingCopy;
 };
 
 /**
@@ -176,7 +188,7 @@ struct ENHANCEDINPUT_API FKeyMappingRow
 {
 	GENERATED_BODY()
 
-	UPROPERTY(VisibleAnywhere, SaveGame, BlueprintReadOnly, Category="Enhanced Input|User Settings")
+	UPROPERTY(VisibleAnywhere, BlueprintReadOnly, Category="Enhanced Input|User Settings")
 	TSet<FPlayerKeyMapping> Mappings;
 
 	/** Returns true if this row has any mappings in it */
@@ -295,8 +307,11 @@ protected:
 	UPROPERTY(BlueprintReadWrite, SaveGame, EditAnywhere, Category="Enhanced Input|User Settings")
 	FText DisplayName;
 	
-	/** A map of "Action Name" to all key mappings associated with it */
-	UPROPERTY(BlueprintReadOnly, SaveGame, EditAnywhere, Category="Enhanced Input|User Settings")
+	/**
+	 * A map of "Action Name" to all key mappings associated with it.
+	 * Note: Dirty mappings will be serialized from UEnhancedInputUserSettings::Serialize
+	 */
+	UPROPERTY(BlueprintReadOnly, Transient, EditAnywhere, Category="Enhanced Input|User Settings")
 	TMap<FName, FKeyMappingRow> PlayerMappedKeys;
 };
 
@@ -355,6 +370,10 @@ public:
 	UFUNCTION(BlueprintCallable, Category="Enhanced Input|User Settings")
 	virtual void ApplySettings();
 
+	/**
+	 * Synchronously save the settings to a hardcoded save game slot. This will work for simple games,
+	 * but if you need to integrate it into an advanced save system you should Serialize this object out with the rest of your save data.
+	 */
 	UFUNCTION(BlueprintCallable, Category="Enhanced Input|User Settings")
 	virtual void SaveSettings();
 	
