Removed TScriptDelegate<FWeakObjectPtr> and TMulticastScriptDelegate<FWeakObjectPtr> usage.

diff --git a/Engine/Plugins/EnhancedInput/Source/EnhancedInput/Public/EnhancedInputComponent.h b/Engine/Plugins/EnhancedInput/Source/EnhancedInput/Public/EnhancedInputComponent.h
--- a/Engine/Plugins/EnhancedInput/Source/EnhancedInput/Public/EnhancedInputComponent.h
+++ b/Engine/Plugins/EnhancedInput/Source/EnhancedInput/Public/EnhancedInputComponent.h
@@ -78,7 +78,7 @@ public:
 
 	/** Binds a script delegate on an arbitrary UObject, hidden for native delegates */
 	template<	typename TSig = TSignature,
-				typename = typename TEnableIf<TIsDerivedFrom<TSig, TScriptDelegate<FWeakObjectPtr> >::IsDerived || TIsDerivedFrom<TSig, TMulticastScriptDelegate<FWeakObjectPtr> >::IsDerived>::Type>
+				typename = typename TEnableIf<TIsDerivedFrom<TSig, FScriptDelegate>::IsDerived || TIsDerivedFrom<TSig, FMulticastScriptDelegate>::IsDerived>::Type>
 	void BindDelegate(UObject* Object, const FName FuncName)
 	{
 		Unbind();
diff --git a/Engine/Source/Runtime/CoreUObject/Public/UObject/ScriptDelegateFwd.h b/Engine/Source/Runtime/CoreUObject/Public/UObject/ScriptDelegateFwd.h
--- a/Engine/Source/Runtime/CoreUObject/Public/UObject/ScriptDelegateFwd.h
+++ b/Engine/Source/Runtime/CoreUObject/Public/UObject/ScriptDelegateFwd.h
@@ -12,5 +12,5 @@ template <typename TWeakPtr = FWeakObjectPtr> class TScriptDelegate;
 template <typename TWeakPtr = FWeakObjectPtr> class TMulticastScriptDelegate;
 
 // Typedef script delegates for convenience.
-typedef TScriptDelegate<FWeakObjectPtr> FScriptDelegate;
-typedef TMulticastScriptDelegate<FWeakObjectPtr> FMulticastScriptDelegate;
+typedef TScriptDelegate<> FScriptDelegate;
+typedef TMulticastScriptDelegate<> FMulticastScriptDelegate;
