Add a runtime log warning when you add or remove a Player Mappable Input config that has been marked as deprecated. This can be turned off in the developer settings by disabling bLogOnDeprecatedConfigUsed

#rnx

diff --git a/Engine/Plugins/EnhancedInput/Source/EnhancedInput/Private/EnhancedInputDeveloperSettings.cpp b/Engine/Plugins/EnhancedInput/Source/EnhancedInput/Private/EnhancedInputDeveloperSettings.cpp
--- a/Engine/Plugins/EnhancedInput/Source/EnhancedInput/Private/EnhancedInputDeveloperSettings.cpp
+++ b/Engine/Plugins/EnhancedInput/Source/EnhancedInput/Private/EnhancedInputDeveloperSettings.cpp
@@ -29,6 +29,7 @@ UEnhancedInputDeveloperSettings::UEnhancedInputDeveloperSettings(const FObjectIn
 	, bEnableUserSettings(false)
 	, bEnableDefaultMappingContexts(true)
 	, bShouldOnlyTriggerLastActionInChord(true)
+	, bLogOnDeprecatedConfigUsed(true)
 	, bEnableWorldSubsystem(false)
 	, bShouldLogAllWorldSubsystemInputs(false)
 {
diff --git a/Engine/Plugins/EnhancedInput/Source/EnhancedInput/Private/EnhancedInputSubsystemInterface.cpp b/Engine/Plugins/EnhancedInput/Source/EnhancedInput/Private/EnhancedInputSubsystemInterface.cpp
--- a/Engine/Plugins/EnhancedInput/Source/EnhancedInput/Private/EnhancedInputSubsystemInterface.cpp
+++ b/Engine/Plugins/EnhancedInput/Source/EnhancedInput/Private/EnhancedInputSubsystemInterface.cpp
@@ -559,9 +559,14 @@ TArray<FKey> IEnhancedInputSubsystemInterface::GetAllPlayerMappedKeys(const FNam
 
 void IEnhancedInputSubsystemInterface::AddPlayerMappableConfig(const UPlayerMappableInputConfig* Config, const FModifyContextOptions& Options)
 {
-	if(Config)
+	if (Config)
 	{
-		for(TPair<TObjectPtr<UInputMappingContext>, int32> Pair : Config->GetMappingContexts())
+		if (GetDefault<UEnhancedInputDeveloperSettings>()->bLogOnDeprecatedConfigUsed && Config->IsDeprecated())
+		{
+			UE_LOG(LogEnhancedInput, Warning, TEXT("The Player Mappable Input Config '%s' is marked as deprecated, but is still being added!"), *Config->GetFName().ToString());
+		}
+
+		for (TPair<TObjectPtr<UInputMappingContext>, int32> Pair : Config->GetMappingContexts())
 		{
 			AddMappingContext(Pair.Key, Pair.Value, Options);
 		}	
@@ -570,8 +575,13 @@ void IEnhancedInputSubsystemInterface::AddPlayerMappableConfig(const UPlayerMapp
 
 void IEnhancedInputSubsystemInterface::RemovePlayerMappableConfig(const UPlayerMappableInputConfig* Config, const FModifyContextOptions& Options)
 {
-	if(Config)
+	if (Config)
 	{
+		if (GetDefault<UEnhancedInputDeveloperSettings>()->bLogOnDeprecatedConfigUsed && Config->IsDeprecated())
+		{
+			UE_LOG(LogEnhancedInput, Warning, TEXT("The Player Mappable Input Config '%s' is marked as deprecated, but is still being removed!"), *Config->GetFName().ToString());
+		}
+
 		for(TPair<TObjectPtr<UInputMappingContext>, int32> Pair : Config->GetMappingContexts())
 		{
 			RemoveMappingContext(Pair.Key, Options);
diff --git a/Engine/Plugins/EnhancedInput/Source/EnhancedInput/Public/EnhancedInputDeveloperSettings.h b/Engine/Plugins/EnhancedInput/Source/EnhancedInput/Public/EnhancedInputDeveloperSettings.h
--- a/Engine/Plugins/EnhancedInput/Source/EnhancedInput/Public/EnhancedInputDeveloperSettings.h
+++ b/Engine/Plugins/EnhancedInput/Source/EnhancedInput/Public/EnhancedInputDeveloperSettings.h
@@ -100,6 +100,12 @@ public:
 	 */
 	UPROPERTY(EditAnywhere, Category = "Enhanced Input", meta=(ConsoleVariable="EnhancedInput.OnlyTriggerLastActionInChord"))
 	uint8 bShouldOnlyTriggerLastActionInChord : 1;
+
+	/**
+	 * If true, then a warning will be logged when a UPlayerMappableInputConfig that has been marked as deprecated is used.
+	 */
+	UPROPERTY(config, EditAnywhere, Category = "Enhanced Input")
+	uint8 bLogOnDeprecatedConfigUsed : 1;
 	
 	/** If true, then the world subsystem will be created. */
 	UPROPERTY(config, EditAnywhere, Category = "Enhanced Input|World Subsystem", meta=(DisplayName="Enable World Subsystem (Experimental)", DisplayPriority = 1))
