Add an Async save option for the Enhanced Input user settings. Make some delegates blueprint bindable, and add another delegate for when the settings are applied so that you can listen for changes in widgets or other places.

#rnx

diff --git a/Engine/Plugins/EnhancedInput/Source/EnhancedInput/Private/EnhancedInputSubsystemInterface.cpp b/Engine/Plugins/EnhancedInput/Source/EnhancedInput/Private/EnhancedInputSubsystemInterface.cpp
--- a/Engine/Plugins/EnhancedInput/Source/EnhancedInput/Private/EnhancedInputSubsystemInterface.cpp
+++ b/Engine/Plugins/EnhancedInput/Source/EnhancedInput/Private/EnhancedInputSubsystemInterface.cpp
@@ -29,12 +29,14 @@ static FAutoConsoleVariableRef GCVarGlobalAxisConfigMode(
 	TEXT("Whether or not to apply Global Axis Config settings. 0 = Default (Mouse Only), 1 = All, 2 = None")
 );
 
+void IEnhancedInputSubsystemInterface::InitalizeUserSettings()
+{
+	// Not every implementer of the EI subsystem wants user settings, so leave it up to them to determine if they want it or not
+}
+
 UEnhancedInputUserSettings* IEnhancedInputSubsystemInterface::GetUserSettings() const
 {
-	if (const UEnhancedPlayerInput* PlayerInput = GetPlayerInput())
-	{
-		return PlayerInput->GetUserSettings();
-	}
+	// Not every implementer of the EI subsystem wants user settings, so leave it up to them to determine if they want it or not
 	return nullptr;
 }
 
@@ -756,8 +758,8 @@ void IEnhancedInputSubsystemInterface::RebuildControlMappings()
 	// Reset the tracking of dependant chord actions on the player input
 	PlayerInput->DependentChordActions.Reset();
 
-	UEnhancedInputUserSettings* UserSettings = GetUserSettings();
-	UEnhancedPlayerMappableKeyProfile* PlayerKeyProfile = UserSettings ? UserSettings->GetCurrentKeyProfile() : nullptr;
+	UEnhancedInputUserSettings* CurrentUserSettings = GetUserSettings();
+	UEnhancedPlayerMappableKeyProfile* PlayerKeyProfile = CurrentUserSettings ? CurrentUserSettings->GetCurrentKeyProfile() : nullptr;
 
 	// An array of keys that are mapped to a given Action.
 	// This is populated by any player mapped keys if they exist, or the default mapping from
diff --git a/Engine/Plugins/EnhancedInput/Source/EnhancedInput/Private/EnhancedInputSubsystems.cpp b/Engine/Plugins/EnhancedInput/Source/EnhancedInput/Private/EnhancedInputSubsystems.cpp
--- a/Engine/Plugins/EnhancedInput/Source/EnhancedInput/Private/EnhancedInputSubsystems.cpp
+++ b/Engine/Plugins/EnhancedInput/Source/EnhancedInput/Private/EnhancedInputSubsystems.cpp
@@ -1,7 +1,7 @@
 // Copyright Epic Games, Inc. All Rights Reserved.
 
 #include "EnhancedInputSubsystems.h"
-
+#include "EnhancedInputModule.h"
 #include "Components/InputComponent.h"
 #include "Engine/LocalPlayer.h"
 #include "EnhancedInputDeveloperSettings.h"
@@ -11,6 +11,7 @@
 #include "Framework/Application/SlateApplication.h"
 #include "GameFramework/PlayerController.h"
 #include "InputMappingContext.h"
+#include "UserSettings/EnhancedInputUserSettings.h"
 
 #include UE_INLINE_GENERATED_CPP_BY_NAME(EnhancedInputSubsystems)
 
@@ -30,7 +31,7 @@ void UEnhancedInputLocalPlayerSubsystem::PlayerControllerChanged(APlayerControll
 	// That means that we can listen for changes to that player's settings!
 	if (GetDefault<UEnhancedInputDeveloperSettings>()->bEnableUserSettings)
 	{
-		BindUserSettingDelegates();	
+		InitalizeUserSettings();
 	}
 }
 
@@ -43,6 +44,42 @@ UEnhancedPlayerInput* UEnhancedInputLocalPlayerSubsystem::GetPlayerInput() const
 	return nullptr;
 }
 
+UEnhancedInputUserSettings* UEnhancedInputLocalPlayerSubsystem::GetUserSettings() const
+{
+	return UserSettings;
+}
+
+void UEnhancedInputLocalPlayerSubsystem::InitalizeUserSettings()
+{
+	// We don't want to be re-creating any user settings, they should share the lifetime with the player input object
+	if (!ensureMsgf(!UserSettings, TEXT("Attempting to initalize the User Settings, but they have already been created!")))
+	{
+		return;
+	}
+
+	if (!GetDefault<UEnhancedInputDeveloperSettings>()->bEnableUserSettings)
+	{
+		UE_LOG(LogEnhancedInput, Verbose, TEXT("bEnableUserSettings is set to false, skipping creation of UEnhancedInputUserSettings!"));
+		return;
+	}
+
+	UEnhancedPlayerInput* PlayerInput = GetPlayerInput();
+	if (!PlayerInput)
+	{
+		UE_LOG(LogEnhancedInput, Error, TEXT("Unable to find a valid player input to initalize user settings with!"));
+		return;
+	}
+	
+	UserSettings = UEnhancedInputUserSettings::LoadOrCreateSettings(PlayerInput);
+
+	// Bind delegates to the user settings if it was successfully created
+	if (ensure(UserSettings))
+	{
+		BindUserSettingDelegates();
+		UE_LOG(LogEnhancedInput, Log, TEXT("Enhanced Input local player subsystem has initalized the user settings!"));
+	}
+}
+
 void UEnhancedInputLocalPlayerSubsystem::ControlMappingsRebuiltThisFrame()
 {
 	ControlMappingsRebuiltDelegate.Broadcast();
diff --git a/Engine/Plugins/EnhancedInput/Source/EnhancedInput/Private/EnhancedPlayerInput.cpp b/Engine/Plugins/EnhancedInput/Source/EnhancedInput/Private/EnhancedPlayerInput.cpp
--- a/Engine/Plugins/EnhancedInput/Source/EnhancedInput/Private/EnhancedPlayerInput.cpp
+++ b/Engine/Plugins/EnhancedInput/Source/EnhancedInput/Private/EnhancedPlayerInput.cpp
@@ -44,46 +44,6 @@ UEnhancedPlayerInput::UEnhancedPlayerInput()
 	}
 }
 
-void UEnhancedPlayerInput::PostInitProperties()
-{
-	Super::PostInitProperties();
-
-	if (HasAnyFlags(RF_ClassDefaultObject))
-	{
-		return;
-	}
-	
-	InitalizeUserSettings();
-}
-
-void UEnhancedPlayerInput::InitalizeUserSettings()
-{
-	// We don't want to be re-creating any user settings, they should share the lifetime with the player input object
-	if (!ensureMsgf(!UserSettings, TEXT("Attempting to initalize the User Settings, but they have already been created!")))
-	{
-		return;
-	}
-	
-	const UEnhancedInputDeveloperSettings* DevSettings = GetDefault<UEnhancedInputDeveloperSettings>();
-	
-	// These are player settings, so only player owned input components can have them.
-	const bool bHasValidPlayerController = GetOuterAPlayerController() != nullptr;
-
-	if (bHasValidPlayerController && DevSettings->bEnableUserSettings && !IsDefaultSubobject())
-	{
-		// Spawn a new instance of the specified user settings class. This is a "NoClear" UPROPERTY so it should always be set
-		const UClass* SettingsClass = DevSettings->UserSettingsClass.Get();
-		if (ensure(SettingsClass))
-		{
-			UserSettings = UEnhancedInputUserSettings::LoadOrCreateSettings(this);
-		}
-	}
-	else
-	{
-		UE_LOG(LogEnhancedInput, Log, TEXT("bEnableUserSettings is set to false, skipping creation of UEnhancedInputUserSettings!"));
-	}
-}
-
 // NOTE: Enum order represents firing priority(lowest to highest) and is important as multiple keys bound to the same action may generate differing trigger event states.
 enum class ETriggerEventInternal : uint8
 {
@@ -295,11 +255,6 @@ float UEnhancedPlayerInput::GetEffectiveTimeDilation() const
 	return 1.0f;
 }
 
-UEnhancedInputUserSettings* UEnhancedPlayerInput::GetUserSettings() const
-{
-	return UserSettings;
-}
-
 void UEnhancedPlayerInput::ProcessInputStack(const TArray<UInputComponent*>& InputComponentStack, const float DeltaTime, const bool bGamePaused)
 {
 	// We need to grab the down states of all keys before calling Super::ProcessInputStack as it will leave bDownPrevious in the same state as bDown (i.e. this frame, not last).
diff --git a/Engine/Plugins/EnhancedInput/Source/EnhancedInput/Private/UserSettings/EnhancedInputUserSettings.cpp b/Engine/Plugins/EnhancedInput/Source/EnhancedInput/Private/UserSettings/EnhancedInputUserSettings.cpp
--- a/Engine/Plugins/EnhancedInput/Source/EnhancedInput/Private/UserSettings/EnhancedInputUserSettings.cpp
+++ b/Engine/Plugins/EnhancedInput/Source/EnhancedInput/Private/UserSettings/EnhancedInputUserSettings.cpp
@@ -19,6 +19,8 @@
 #include "EnhancedInputLibrary.h"
 #include "Algo/Find.h"
 #include "NativeGameplayTags.h"
+#include "PlatformFeatures.h"
+#include "SaveGameSystem.h"
 
 #define LOCTEXT_NAMESPACE "EnhancedInputMappableUserSettings"
 
@@ -39,7 +41,7 @@ namespace UE::EnhancedInput
 	{
 		if (PlayerInput)
 		{
-			if (APlayerController* PC = Cast<APlayerController>(PlayerInput->GetOuter()))
+			if (APlayerController* PC = PlayerInput->GetOuterAPlayerController())
 			{
 				return PC->GetLocalPlayer();
 			}
@@ -414,14 +416,22 @@ UEnhancedInputUserSettings* UEnhancedInputUserSettings::LoadOrCreateSettings(UEn
 		USaveGame* Slot = UGameplayStatics::LoadGameFromSlot(UE::EnhancedInput::SETTINGS_SLOT_NAME, LocalPlayer->GetLocalPlayerIndex());
 		Settings = Cast<UEnhancedInputUserSettings>(Slot);
 	}
-	
+
+	// If there is no settings save game object, then we can create on
+	// based on the class type set in the developer settings
 	if (Settings == nullptr)
 	{
-		Settings = Cast<UEnhancedInputUserSettings>(UGameplayStatics::CreateSaveGameObject(UEnhancedInputUserSettings::StaticClass()));
+		const UEnhancedInputDeveloperSettings* DevSettings = GetDefault<UEnhancedInputDeveloperSettings>();
+		UClass* SettingsClass = DevSettings->UserSettingsClass ? DevSettings->UserSettingsClass.Get() : UEnhancedInputDeveloperSettings::StaticClass();
+		
+		Settings = Cast<UEnhancedInputUserSettings>(UGameplayStatics::CreateSaveGameObject(SettingsClass));
 	}
 
-	Settings->Initialize(PlayerInput);
-	Settings->ApplySettings();
+	if (ensure(Settings))
+	{
+		Settings->Initialize(PlayerInput);
+    	Settings->ApplySettings();	
+	}
 
 	return Settings;
 }
@@ -447,6 +457,8 @@ void UEnhancedInputUserSettings::ApplySettings()
 {
 	ensureMsgf(GetPlayerInput(), TEXT("UEnhancedInputUserSettings is missing a player input!"));
 	UE_LOG(LogEnhancedInput, Verbose, TEXT("Enhanced Input User Settings applied!"));
+
+	OnSettingsApplied.Broadcast();
 }
 
 void UEnhancedInputUserSettings::SaveSettings()
@@ -464,6 +476,33 @@ void UEnhancedInputUserSettings::SaveSettings()
 	}
 }
 
+void UEnhancedInputUserSettings::AsyncSaveSettings()
+{
+	if (const ULocalPlayer* OwningPlayer = GetLocalPlayer())
+	{
+		FAsyncSaveGameToSlotDelegate SavedDelegate;
+		SavedDelegate.BindUObject(this, &UEnhancedInputUserSettings::OnAsyncSaveComplete);
+		
+		UGameplayStatics::AsyncSaveGameToSlot(this, UE::EnhancedInput::SETTINGS_SLOT_NAME, OwningPlayer->GetLocalPlayerIndex(), SavedDelegate);
+	}
+	else
+	{
+		UE_LOG(LogEnhancedInput, Error, TEXT("[UEnhancedInputUserSettings::AsyncSaveSettings] Failed async save, there is no owning local player!"));
+	}
+}
+
+void UEnhancedInputUserSettings::OnAsyncSaveComplete(const FString& SlotName, const int32 UserIndex, bool bSuccess)
+{
+	if (bSuccess)
+	{
+		UE_LOG(LogEnhancedInput, Log, TEXT("UEnhancedInputUserSettings::OnAsyncSaveComplete] Async save of slot '%s' for user index '%d' completed successfully!"), *SlotName, UserIndex);
+	}
+	else
+	{
+		UE_LOG(LogEnhancedInput, Error, TEXT("[UEnhancedInputUserSettings::OnAsyncSaveComplete] Failed async save of slot '%s' for user index '%d'"), *SlotName, UserIndex);
+	}
+}
+
 namespace UE::EnhancedInput
 {
 	static const int32 GPlayerMappableSaveVersion = 1;
@@ -646,13 +685,12 @@ void UEnhancedInputUserSettings::MapPlayerKey(const FMapPlayerKeyArgs& InArgs, F
 		// Then set the player mapped key
 		FoundMapping->SetCurrentKey(InArgs.NewKey);
 		OnSettingsChanged.Broadcast(this);
-		return;
 	}
 	// If it doesn't exist, then we need to make it if there is a valid action name
 	else if (FKeyMappingRow* MappingRow = KeyProfile->PlayerMappedKeys.Find(InArgs.ActionName))
 	{
-		// If one doesn't exist, then we need to create a new mapping in the given
-		// slot. 
+		// If one doesn't exist, then we need to create a new mapping in the given slot.
+		// 
 		// In order to populate the default values correctly, we only do this if we know that
 		// mappings exist for it
 		const int32 NumMappings = MappingRow->Mappings.Num();
@@ -708,11 +746,14 @@ void UEnhancedInputUserSettings::UnMapPlayerKey(const FMapPlayerKeyArgs& InArgs,
 	{
 		// Then set the player mapped key
 		FoundMapping->ResetToDefault();
-		OnSettingsChanged.Broadcast(this);
+
+		// The settings have only changed if the mapping is dirty now
+		if (FoundMapping->IsDirty())
+		{
+			OnSettingsChanged.Broadcast(this);
+		}
 		
 		UE_LOG(LogEnhancedInput, Verbose, TEXT("[UEnhancedInputUserSettings::MapPlayerKey] Reset keymapping to default: '%s'"), *FoundMapping->ToString());
-		
-		return;
 	}
 	// if a mapping doesn't exist, then we can't unmap it
 	else
diff --git a/Engine/Plugins/EnhancedInput/Source/EnhancedInput/Public/EnhancedInputSubsystemInterface.h b/Engine/Plugins/EnhancedInput/Source/EnhancedInput/Public/EnhancedInputSubsystemInterface.h
--- a/Engine/Plugins/EnhancedInput/Source/EnhancedInput/Public/EnhancedInputSubsystemInterface.h
+++ b/Engine/Plugins/EnhancedInput/Source/EnhancedInput/Public/EnhancedInputSubsystemInterface.h
@@ -84,6 +84,13 @@ public:
 
 protected:
 	
+	/**
+	 * Create a new user settings object if it is enabled in the EI developer settings.
+	 *
+	 * Not every enhanced input subsystem needs user settings, so this is an optional feature.
+	 */
+	virtual void InitalizeUserSettings();
+	
 	/** Binds to any delegates of interest on the UEnhancedInputUserSettings if they are enabled in the developer settings. */
 	virtual void BindUserSettingDelegates();
 	
diff --git a/Engine/Plugins/EnhancedInput/Source/EnhancedInput/Public/EnhancedInputSubsystems.h b/Engine/Plugins/EnhancedInput/Source/EnhancedInput/Public/EnhancedInputSubsystems.h
--- a/Engine/Plugins/EnhancedInput/Source/EnhancedInput/Public/EnhancedInputSubsystems.h
+++ b/Engine/Plugins/EnhancedInput/Source/EnhancedInput/Public/EnhancedInputSubsystems.h
@@ -28,6 +28,8 @@ public:
 	
 	// Begin IEnhancedInputSubsystemInterface
 	virtual UEnhancedPlayerInput* GetPlayerInput() const override;
+	virtual UEnhancedInputUserSettings* GetUserSettings() const override;
+	virtual void InitalizeUserSettings() override;
 	virtual void ControlMappingsRebuiltThisFrame() override;
 	// End IEnhancedInputSubsystemInterface
 
@@ -39,6 +41,12 @@ public:
 	 */
 	UPROPERTY(BlueprintAssignable, DisplayName=OnControlMappingsRebuilt, Category = "Input")
 	FOnControlMappingsRebuilt ControlMappingsRebuiltDelegate;
+
+protected:
+    	
+	/** The user settings for this subsystem used to store each user's input related settings */
+	UPROPERTY()
+	TObjectPtr<UEnhancedInputUserSettings> UserSettings;
 };
 
 /**
diff --git a/Engine/Plugins/EnhancedInput/Source/EnhancedInput/Public/EnhancedPlayerInput.h b/Engine/Plugins/EnhancedInput/Source/EnhancedInput/Public/EnhancedPlayerInput.h
--- a/Engine/Plugins/EnhancedInput/Source/EnhancedInput/Public/EnhancedPlayerInput.h
+++ b/Engine/Plugins/EnhancedInput/Source/EnhancedInput/Public/EnhancedPlayerInput.h
@@ -49,8 +49,6 @@ public:
 
 	UEnhancedPlayerInput();
 
-	virtual void PostInitProperties() override;
-
 	/**
 	* Returns the action instance data for the given input action if there is any. Returns nullptr if the action is not available.
 	*/
@@ -71,14 +69,8 @@ public:
 
 	/** Returns the Time Dilation value that is currently effecting this input. */
 	float GetEffectiveTimeDilation() const;
-
-	/** Returns the current user settings for this player input. Will be null if the bEnableUserSettings flag is false on the developer settings */
-	UEnhancedInputUserSettings* GetUserSettings() const;
 	
 protected:
-
-	/** Create a new user settings object if it is enabled in the EI developer settings */
-	virtual void InitalizeUserSettings();
 	
 	// Causes key to be consumed if it is affecting an action.
 	virtual bool IsKeyHandledByAction(FKey Key) const override;
@@ -167,11 +159,6 @@ private:
 	 */
 	TArray<FDependentChordTracker> DependentChordActions;
 
-protected:
-	
-	/** The user settings for this subsystem used to store each user's input related settings */
-	UPROPERTY()
-	TObjectPtr<UEnhancedInputUserSettings> UserSettings;
 
 private:
 
diff --git a/Engine/Plugins/EnhancedInput/Source/EnhancedInput/Public/UserSettings/EnhancedInputUserSettings.h b/Engine/Plugins/EnhancedInput/Source/EnhancedInput/Public/UserSettings/EnhancedInputUserSettings.h
--- a/Engine/Plugins/EnhancedInput/Source/EnhancedInput/Public/UserSettings/EnhancedInputUserSettings.h
+++ b/Engine/Plugins/EnhancedInput/Source/EnhancedInput/Public/UserSettings/EnhancedInputUserSettings.h
@@ -367,6 +367,10 @@ public:
 	static UEnhancedInputUserSettings* LoadOrCreateSettings(UEnhancedPlayerInput* PlayerInput);
 	virtual void Initialize(UEnhancedPlayerInput* PlayerInput);
 
+	/**
+	 * Apply any custom input settings to your user. By default, this will just broadcast the OnSettingsApplied delegate
+	 * which is a useful hook to maybe rebuild some UI or do other user facing updates.
+	 */
 	UFUNCTION(BlueprintCallable, Category="Enhanced Input|User Settings")
 	virtual void ApplySettings();
 
@@ -376,6 +380,21 @@ public:
 	 */
 	UFUNCTION(BlueprintCallable, Category="Enhanced Input|User Settings")
 	virtual void SaveSettings();
+
+	/**
+	 * Asynchronously save the settings to a hardcoded save game slot. This will work for simple games,
+	 * but if you need to integrate it into an advanced save system you should Serialize this object out with the rest of your save data.
+	 *
+	 * OnAsyncSaveComplete will be called upon save completion.
+	 */
+	UFUNCTION(BlueprintCallable, Category="Enhanced Input|User Settings")
+	virtual void AsyncSaveSettings();
+
+protected:
+
+	virtual void OnAsyncSaveComplete(const FString& SlotName, const int32 UserIndex, bool bSuccess);
+
+public:
 	
 	UEnhancedPlayerInput* GetPlayerInput() const;
 	ULocalPlayer* GetLocalPlayer() const;
@@ -383,7 +402,13 @@ public:
 	
 	/** Fired when the user settings have changed, such as their key mappings. */
 	DECLARE_DYNAMIC_MULTICAST_DELEGATE_OneParam(FEnhancedInputUserSettingsChanged, UEnhancedInputUserSettings*, Settings);
+	UPROPERTY(BlueprintAssignable, Category = "Enhanced Input|User Settings")
 	FEnhancedInputUserSettingsChanged OnSettingsChanged;
+	
+	/** Called after the settings have been applied from the ApplySettings call. */
+	DECLARE_DYNAMIC_MULTICAST_DELEGATE(FEnhancedInputUserSettingsApplied);
+	UPROPERTY(BlueprintAssignable, Category = "Enhanced Input|User Settings")
+	FEnhancedInputUserSettingsApplied OnSettingsApplied;
 
 	// Remappable keys API
 
