Rename the FMapPlayerKeyArgs param from ActionName to MappingName to be more consistent with the IMC namin, along with other functions that reference it as an "Aciton Name"

Make the Register/Unregister IMC functions virtual.

#rnx

diff --git a/Engine/Plugins/EnhancedInput/Source/EnhancedInput/Private/EnhancedInputSubsystemInterface.cpp b/Engine/Plugins/EnhancedInput/Source/EnhancedInput/Private/EnhancedInputSubsystemInterface.cpp
--- a/Engine/Plugins/EnhancedInput/Source/EnhancedInput/Private/EnhancedInputSubsystemInterface.cpp
+++ b/Engine/Plugins/EnhancedInput/Source/EnhancedInput/Private/EnhancedInputSubsystemInterface.cpp
@@ -784,7 +784,7 @@ void IEnhancedInputSubsystemInterface::RebuildControlMappings()
 			// See if there are any player mapped keys to this action
 			if (PlayerKeyProfile && GetDefault<UEnhancedInputDeveloperSettings>()->bEnableUserSettings)
 			{
-				PlayerKeyProfile->GetKeysMappedToAction(Mapping.GetMappingName(), MappedKeysToActionName);
+				PlayerKeyProfile->GetMappedKeysInRow(Mapping.GetMappingName(), MappedKeysToActionName);
 				bIsPlayerMapping = !MappedKeysToActionName.IsEmpty();
 			}
 			else if (const TMap<FPlayerMappableKeySlot, FKey>* PlayerMappedKeySlots = PlayerMappedSettings.Find(Mapping.GetMappingName()))
diff --git a/Engine/Plugins/EnhancedInput/Source/EnhancedInput/Private/UserSettings/EnhancedInputUserSettings.cpp b/Engine/Plugins/EnhancedInput/Source/EnhancedInput/Private/UserSettings/EnhancedInputUserSettings.cpp
--- a/Engine/Plugins/EnhancedInput/Source/EnhancedInput/Private/UserSettings/EnhancedInputUserSettings.cpp
+++ b/Engine/Plugins/EnhancedInput/Source/EnhancedInput/Private/UserSettings/EnhancedInputUserSettings.cpp
@@ -21,6 +21,7 @@
 #include "NativeGameplayTags.h"
 #include "PlatformFeatures.h"
 #include "SaveGameSystem.h"
+#include "EnhancedActionKeyMapping.h"
 
 #define LOCTEXT_NAMESPACE "EnhancedInputMappableUserSettings"
 
@@ -32,7 +33,7 @@ namespace UE::EnhancedInput
 	UE_DEFINE_GAMEPLAY_TAG_STATIC(TAG_DefaultProfileIdentifier, "InputUserSettings.Profiles.Default");
 	static const FText DefaultProfileDisplayName = LOCTEXT("Default_Profile_name", "Default Profile");
 	
-	UE_DEFINE_GAMEPLAY_TAG_STATIC(TAG_InvalidActionName, "InputUserSettings.FailureReasons.InvalidActionName");
+	UE_DEFINE_GAMEPLAY_TAG_STATIC(TAG_InvalidMappingName, "InputUserSettings.FailureReasons.InvalidMappingName");
 	UE_DEFINE_GAMEPLAY_TAG_STATIC(TAG_NoKeyProfile, "InputUserSettings.FailureReasons.NoKeyProfile");
 	UE_DEFINE_GAMEPLAY_TAG_STATIC(TAG_NoMatchingMappings, "InputUserSettings.FailureReasons.NoMatchingMappings");
 	
@@ -74,7 +75,7 @@ namespace UE::EnhancedInput
 // FMapPlayerKeyArgs
 
 FMapPlayerKeyArgs::FMapPlayerKeyArgs()
-	: ActionName(NAME_None)
+	: MappingName(NAME_None)
 	, Slot(EPlayerMappableKeySlot::Unspecified)
 	, NewKey(EKeys::Invalid)
 	, HardwareDeviceId(NAME_None)
@@ -86,7 +87,7 @@ FMapPlayerKeyArgs::FMapPlayerKeyArgs()
 // FPlayerKeyMapping
 
 FPlayerKeyMapping::FPlayerKeyMapping()
-	: ActionName(NAME_None)
+	: MappingName(NAME_None)
 	, DisplayName(FText::GetEmpty())
 	, Slot(EPlayerMappableKeySlot::Unspecified)
 	, DefaultKey(EKeys::Invalid)
@@ -96,7 +97,7 @@ FPlayerKeyMapping::FPlayerKeyMapping()
 }
 
 FPlayerKeyMapping::FPlayerKeyMapping(const FEnhancedActionKeyMapping& OriginalMapping, EPlayerMappableKeySlot InSlot /* = EPlayerMappableKeySlot::Unspecified */)
-	: ActionName(OriginalMapping.GetMappingName())
+	: MappingName(OriginalMapping.GetMappingName())
 	, DisplayName(OriginalMapping.GetDisplayName())
 	, Slot(InSlot)
 	, DefaultKey(OriginalMapping.Key)
@@ -116,7 +117,7 @@ bool FPlayerKeyMapping::IsCustomized() const
 
 bool FPlayerKeyMapping::IsValid() const
 {
-	return ActionName.IsValid() && CurrentKey.IsValid();
+	return MappingName.IsValid() && CurrentKey.IsValid();
 }
 
 const FKey& FPlayerKeyMapping::GetCurrentKey() const
@@ -135,16 +136,16 @@ FString FPlayerKeyMapping::ToString() const
 	check(PlayerMappableEnumClass);
 	return
 		FString::Printf(TEXT("Action Name: '%s'  Slot: '%s'  Default Key: '%s'  Player Mapped Key: '%s'  HardwareDevice:  '%s'"),
-			*ActionName.ToString(),
+			*MappingName.ToString(),
 			*PlayerMappableEnumClass->GetNameStringByValue(static_cast<int64>(Slot)),
 			*DefaultKey.ToString(),
 			*CurrentKey.ToString(),
 			*HardwareDeviceId.ToString());
 }
 
-const FName FPlayerKeyMapping::GetActionName() const
+const FName FPlayerKeyMapping::GetMappingName() const
 {
-	return ActionName;
+	return MappingName;
 }
 
 const FText& FPlayerKeyMapping::GetDisplayName() const
@@ -165,7 +166,7 @@ const FHardwareDeviceIdentifier& FPlayerKeyMapping::GetHardwareDeviceId() const
 uint32 GetTypeHash(const FPlayerKeyMapping& InMapping)
 {
 	uint32 Hash = 0;
-	Hash = HashCombine(Hash, GetTypeHash(InMapping.ActionName));
+	Hash = HashCombine(Hash, GetTypeHash(InMapping.MappingName));
 	Hash = HashCombine(Hash, GetTypeHash(InMapping.Slot));
 	Hash = HashCombine(Hash, GetTypeHash(InMapping.CurrentKey));
 	Hash = HashCombine(Hash, GetTypeHash(InMapping.HardwareDeviceId));
@@ -192,7 +193,7 @@ void FPlayerKeyMapping::SetCurrentKey(const FKey& NewKey)
 
 void FPlayerKeyMapping::UpdateOriginalKey(const FEnhancedActionKeyMapping& OriginalMapping)
 {
-	ensure(OriginalMapping.IsPlayerMappable() && OriginalMapping.GetMappingName() == ActionName);
+	ensure(OriginalMapping.IsPlayerMappable() && OriginalMapping.GetMappingName() == MappingName);
 	
 	DefaultKey = OriginalMapping.Key;
 	DisplayName = OriginalMapping.GetDisplayName();
@@ -208,7 +209,7 @@ void FPlayerKeyMapping::UpdateOriginalKey(const FEnhancedActionKeyMapping& Origi
 bool FPlayerKeyMapping::operator==(const FPlayerKeyMapping& Other) const
 {
 	return
-	       ActionName			== Other.ActionName
+	       MappingName			== Other.MappingName
 		&& Slot					== Other.Slot
 		&& HardwareDeviceId		== Other.HardwareDeviceId
 		&& CurrentKey			== Other.CurrentKey
@@ -267,14 +268,14 @@ const FText& UEnhancedPlayerMappableKeyProfile::GetProfileDisplayName() const
 	return DisplayName;
 }
 
-const TMap<FName, FKeyMappingRow>& UEnhancedPlayerMappableKeyProfile::GetPlayerMappedActions() const
+const TMap<FName, FKeyMappingRow>& UEnhancedPlayerMappableKeyProfile::GetPlayerMappingRows() const
 {
 	return PlayerMappedKeys;
 }
 
-void UEnhancedPlayerMappableKeyProfile::ResetActionMappingsToDefault(const FName InActionName)
+void UEnhancedPlayerMappableKeyProfile::ResetMappingToDefault(const FName InMappingName)
 {
-	if (FKeyMappingRow* MappingRow = FindKeyMappingRowMutable(InActionName))
+	if (FKeyMappingRow* MappingRow = FindKeyMappingRowMutable(InMappingName))
 	{
 		for (FPlayerKeyMapping& Mapping : MappingRow->Mappings)
 		{
@@ -283,14 +284,14 @@ void UEnhancedPlayerMappableKeyProfile::ResetActionMappingsToDefault(const FName
 	}
 }
 
-FKeyMappingRow* UEnhancedPlayerMappableKeyProfile::FindKeyMappingRowMutable(const FName InActionName)
+FKeyMappingRow* UEnhancedPlayerMappableKeyProfile::FindKeyMappingRowMutable(const FName InMappingName)
 {
-	return const_cast<FKeyMappingRow*>(FindKeyMappingRow(InActionName));
+	return const_cast<FKeyMappingRow*>(FindKeyMappingRow(InMappingName));
 }
 
-const FKeyMappingRow* UEnhancedPlayerMappableKeyProfile::FindKeyMappingRow(const FName InActionName) const
+const FKeyMappingRow* UEnhancedPlayerMappableKeyProfile::FindKeyMappingRow(const FName InMappingName) const
 {
-	return PlayerMappedKeys.Find(InActionName);
+	return PlayerMappedKeys.Find(InMappingName);
 }
 
 void UEnhancedPlayerMappableKeyProfile::DumpProfileToLog() const
@@ -317,11 +318,11 @@ FString UEnhancedPlayerMappableKeyProfile::ToString() const
 	return Builder.ToString();
 }
 
-int32 UEnhancedPlayerMappableKeyProfile::GetKeysMappedToAction(const FName ActionName, TArray<FKey>& OutKeys) const
+int32 UEnhancedPlayerMappableKeyProfile::GetMappedKeysInRow(const FName MappingName, TArray<FKey>& OutKeys) const
 {
 	OutKeys.Reset();
 	
-	if (const FKeyMappingRow* MappingRow = FindKeyMappingRow(ActionName))
+	if (const FKeyMappingRow* MappingRow = FindKeyMappingRow(MappingName))
 	{
 		for (const FPlayerKeyMapping& Mapping : MappingRow->Mappings)
 		{
@@ -330,14 +331,14 @@ int32 UEnhancedPlayerMappableKeyProfile::GetKeysMappedToAction(const FName Actio
 	}
 	else
 	{
-		UE_LOG(LogEnhancedInput, Warning, TEXT("Player Mappable Key Profile '%s' doesn't have any mappings for action '%s'"), *ProfileIdentifier.ToString(), *ActionName.ToString());
+		UE_LOG(LogEnhancedInput, Warning, TEXT("Player Mappable Key Profile '%s' doesn't have any mappings for action '%s'"), *ProfileIdentifier.ToString(), *MappingName.ToString());
 	}
 	return OutKeys.Num();
 }
 
-int32 UEnhancedPlayerMappableKeyProfile::GetActionsMappedToKey(const FKey& InKey, TArray<FName>& OutMappedActionNames) const
+int32 UEnhancedPlayerMappableKeyProfile::GetMappingNamesForKey(const FKey& InKey, TArray<FName>& OutMappingNames) const
 {
-	OutMappedActionNames.Reset();
+	OutMappingNames.Reset();
 
 	for (const TPair<FName, FKeyMappingRow>& Pair : PlayerMappedKeys)
 	{
@@ -347,13 +348,13 @@ int32 UEnhancedPlayerMappableKeyProfile::GetActionsMappedToKey(const FKey& InKey
 			{
 				// We know that this action has the key mapped to it, so there is no need to continue checking
 				// the rest of it's mappings
-				OutMappedActionNames.Add(Pair.Key);
+				OutMappingNames.Add(Pair.Key);
 				break;
 			}
 		}
 	}
 
-	return OutMappedActionNames.Num();
+	return OutMappingNames.Num();
 }
 
 bool FKeyMappingRow::HasAnyMappings() const
@@ -370,7 +371,7 @@ void UEnhancedPlayerMappableKeyProfile::Serialize(FArchive& Ar)
 FPlayerKeyMapping* UEnhancedPlayerMappableKeyProfile::FindKeyMapping(const FMapPlayerKeyArgs& InArgs) const
 {
 	// Get the current mappings for the desired action name.
-	if (const FKeyMappingRow* MappingRow = PlayerMappedKeys.Find(InArgs.ActionName))
+	if (const FKeyMappingRow* MappingRow = PlayerMappedKeys.Find(InArgs.MappingName))
 	{
 		// If mapping already exists for the given slot and hardware device, then we can
 		// just change that key
@@ -582,7 +583,7 @@ void UEnhancedInputUserSettings::Serialize(FArchive& Ar)
 					{
 						Header.DirtyMappings.Push(
 							{
-								/* .ActionName = */ Mapping.ActionName,
+								/* .MappingName = */ Mapping.MappingName,
 								/* .CurrentKeyName = */ Mapping.CurrentKey.GetFName(),
 								/* .HardwareDeviceId = */ Mapping.HardwareDeviceId, 
 								/* .Slot = */ Mapping.Slot
@@ -612,7 +613,7 @@ void UEnhancedInputUserSettings::Serialize(FArchive& Ar)
 					FKeyMappingRow& MappingRow = NewProfile->PlayerMappedKeys.FindOrAdd(SavedKeyData.ActionName);
 
 					FPlayerKeyMapping PlayerMapping = {};
-					PlayerMapping.ActionName = SavedKeyData.ActionName;
+					PlayerMapping.MappingName = SavedKeyData.ActionName;
 					PlayerMapping.CurrentKey = FKey(SavedKeyData.CurrentKeyName);
 					PlayerMapping.Slot = SavedKeyData.Slot;
 					PlayerMapping.HardwareDeviceId = SavedKeyData.HardwareDeviceId;
@@ -665,9 +666,9 @@ ULocalPlayer* UEnhancedInputUserSettings::GetLocalPlayer() const
 
 void UEnhancedInputUserSettings::MapPlayerKey(const FMapPlayerKeyArgs& InArgs, FGameplayTagContainer& FailureReason)
 {
-	if (!InArgs.ActionName.IsValid())
+	if (!InArgs.MappingName.IsValid())
 	{
-		FailureReason.AddTag(UE::EnhancedInput::TAG_InvalidActionName);
+		FailureReason.AddTag(UE::EnhancedInput::TAG_InvalidMappingName);
 		return;
 	}
 	
@@ -687,7 +688,7 @@ void UEnhancedInputUserSettings::MapPlayerKey(const FMapPlayerKeyArgs& InArgs, F
 		OnSettingsChanged.Broadcast(this);
 	}
 	// If it doesn't exist, then we need to make it if there is a valid action name
-	else if (FKeyMappingRow* MappingRow = KeyProfile->PlayerMappedKeys.Find(InArgs.ActionName))
+	else if (FKeyMappingRow* MappingRow = KeyProfile->PlayerMappedKeys.Find(InArgs.MappingName))
 	{
 		// If one doesn't exist, then we need to create a new mapping in the given slot.
 		// 
@@ -700,7 +701,7 @@ void UEnhancedInputUserSettings::MapPlayerKey(const FMapPlayerKeyArgs& InArgs, F
 			
 			// Add a default mapping to this row
 			FPlayerKeyMapping PlayerMappingData = {};
-			PlayerMappingData.ActionName = InArgs.ActionName;
+			PlayerMappingData.MappingName = InArgs.MappingName;
 			PlayerMappingData.Slot = InArgs.Slot;
 
 			// If there is some valid hardware then keep track of that
@@ -728,9 +729,9 @@ void UEnhancedInputUserSettings::MapPlayerKey(const FMapPlayerKeyArgs& InArgs, F
 
 void UEnhancedInputUserSettings::UnMapPlayerKey(const FMapPlayerKeyArgs& InArgs, FGameplayTagContainer& FailureReason)
 {
-	if (!InArgs.ActionName.IsValid())
+	if (!InArgs.MappingName.IsValid())
 	{
-		FailureReason.AddTag(UE::EnhancedInput::TAG_InvalidActionName);
+		FailureReason.AddTag(UE::EnhancedInput::TAG_InvalidMappingName);
 		return;
 	}
 	
@@ -762,11 +763,11 @@ void UEnhancedInputUserSettings::UnMapPlayerKey(const FMapPlayerKeyArgs& InArgs,
 	}
 }
 
-const TSet<FPlayerKeyMapping>& UEnhancedInputUserSettings::FindMappingsForAction(const FName ActionName) const
+const TSet<FPlayerKeyMapping>& UEnhancedInputUserSettings::FindMappingsInRow(const FName MappingName) const
 {
 	if (UEnhancedPlayerMappableKeyProfile* KeyProfile = GetCurrentKeyProfile())
 	{
-		FKeyMappingRow& ExistingMappings = KeyProfile->PlayerMappedKeys.FindOrAdd(ActionName);
+		FKeyMappingRow& ExistingMappings = KeyProfile->PlayerMappedKeys.FindOrAdd(MappingName);
 		return ExistingMappings.Mappings;
 	}
 	
@@ -778,7 +779,7 @@ const TSet<FPlayerKeyMapping>& UEnhancedInputUserSettings::FindMappingsForAction
 
 const FPlayerKeyMapping* UEnhancedInputUserSettings::FindCurrentMappingForSlot(const FName ActionName, const EPlayerMappableKeySlot InSlot) const
 {
-	const TSet<FPlayerKeyMapping>& AllMappings = FindMappingsForAction(ActionName);
+	const TSet<FPlayerKeyMapping>& AllMappings = FindMappingsInRow(ActionName);
 	for (const FPlayerKeyMapping& Mapping : AllMappings)
 	{
 		if (Mapping.Slot == InSlot)
diff --git a/Engine/Plugins/EnhancedInput/Source/EnhancedInput/Public/UserSettings/EnhancedInputUserSettings.h b/Engine/Plugins/EnhancedInput/Source/EnhancedInput/Public/UserSettings/EnhancedInputUserSettings.h
--- a/Engine/Plugins/EnhancedInput/Source/EnhancedInput/Public/UserSettings/EnhancedInputUserSettings.h
+++ b/Engine/Plugins/EnhancedInput/Source/EnhancedInput/Public/UserSettings/EnhancedInputUserSettings.h
@@ -6,7 +6,6 @@
 #include "UObject/Object.h"
 #include "InputCoreTypes.h"					// For FKey
 #include "GameFramework/InputSettings.h"	// For FHardwareDeviceIdentifier
-#include "EnhancedActionKeyMapping.h"
 #include "GameplayTagContainer.h"
 
 #include "EnhancedInputUserSettings.generated.h"
@@ -14,12 +13,13 @@
 class UInputMappingContext;
 class UEnhancedPlayerInput;
 class ULocalPlayer;
+struct FEnhancedActionKeyMapping;
 
 /**
  * The "Slot" that a player mappable key is in.
- * Used by UI to allow for multiple keys to be bound by the player for a single action
+ * Used by UI to allow for multiple keys to be bound by the player for a single player mapping
  * 
- * | <Action Name>  | Slot 1 | Slot 2 | Slot 3 | Slot.... N |
+ * | <Mapping Name>  | Slot 1 | Slot 2 | Slot 3 | Slot.... N |
  */
 UENUM(BlueprintType)
 enum class EPlayerMappableKeySlot : uint8 
@@ -50,17 +50,17 @@ struct ENHANCEDINPUT_API FMapPlayerKeyArgs
 	FMapPlayerKeyArgs();	
 	
 	/**
-	 * The name of the action for this key. This is either the default mapping name from an Input Action asset, or one
-	 * that is overriden in the Input Mapping Context.
+	 * The name of the mapping for this key. This is either the default mapping name from an Input Action asset, or one
+	 * that is overridden in the Input Mapping Context.
 	 */
 	UPROPERTY(EditAnywhere, BlueprintReadWrite, Category="Enhanced Input|User Settings")
-	FName ActionName;
+	FName MappingName;
 
 	/** What slot this key mapping is for */
 	UPROPERTY(EditAnywhere, BlueprintReadWrite, Category="Enhanced Input|User Settings")
 	EPlayerMappableKeySlot Slot;
 
-	/** The new Key that this action should be mapped to */
+	/** The new Key that this should be mapped to */
 	UPROPERTY(EditAnywhere, BlueprintReadWrite, Category="Enhanced Input|User Settings")
 	FKey NewKey;
 	
@@ -112,10 +112,10 @@ public:
 	FString ToString() const;
 
 	/**
-	 * The unique FName associated with this action. This is defined by this mappings owning Input Action
+	 * The unique FName associated with this mapping. This is defined by this mappings owning Input Action
 	 * or the individual Enhanced Action Key Mapping if it is overriden
 	 */
-	const FName GetActionName() const;
+	const FName GetMappingName() const;
 
 	/** The localized display name to use for this mapping */
 	const FText& GetDisplayName() const;
@@ -148,11 +148,11 @@ public:
 
 protected:
 	
-	/** The name of the action for this key */
+	/** The name of the mapping for this key */
 	UPROPERTY(VisibleAnywhere, BlueprintReadOnly, Category="Enhanced Input|User Settings")
-	FName ActionName;
+	FName MappingName;
 	
-	/** Localized display name of this action */
+	/** Localized display name of this mapping */
 	UPROPERTY(Transient, VisibleAnywhere, BlueprintReadOnly, Category="Enhanced Input|User Settings")
 	FText DisplayName;
 
@@ -168,7 +168,7 @@ protected:
 	UPROPERTY(Transient, VisibleAnywhere, BlueprintReadOnly, Category="Enhanced Input|User Settings", meta = (AllowPrivateAccess = "True"))
 	FKey DefaultKey;
 	
-	/** The key that the player has mapped this action to */
+	/** The key that the player has mapped to */
 	UPROPERTY(VisibleAnywhere,  BlueprintReadOnly, Category="Enhanced Input|User Settings", meta = (AllowPrivateAccess = "True"))
 	FKey CurrentKey;
 
@@ -178,9 +178,9 @@ protected:
 };
 
 /**
- * Stores all mappings bound to a single action.
+ * Stores all mappings bound to a single mapping name.
  *
- * Since a single action can have multiple bindings to it and this system should be Blueprint friendly,
+ * Since a single mapping can have multiple bindings to it and this system should be Blueprint friendly,
  * this needs to be a struct (blueprint don't support nested containers).
  */
 USTRUCT(BlueprintType)
@@ -229,20 +229,20 @@ public:
 	/**
 	 * Get all known key mappings for this profile.
 	 *
-	 * This returns a map of "Action Name" -> Mappings to that action
+	 * This returns a map of "Mapping Name" -> Player Mappings to that name
 	 */
 	UFUNCTION(BlueprintCallable, BlueprintPure, Category="Enhanced Input|User Settings")
-	const TMap<FName, FKeyMappingRow>& GetPlayerMappedActions() const;
+	const TMap<FName, FKeyMappingRow>& GetPlayerMappingRows() const;
 
-	/** Resets every player key mapping to this action back to it's default value */
+	/** Resets every player key mapping to this mapping back to it's default value */
 	UFUNCTION(BlueprintCallable, Category="Enhanced Input|User Settings")
-	void ResetActionMappingsToDefault(const FName InActionName);
+	void ResetMappingToDefault(const FName InMappingName);
 	
-	/** Get all the key mappings associated with the given action name on this profile */
-	FKeyMappingRow* FindKeyMappingRowMutable(const FName InActionName);
+	/** Get all the key mappings associated with the given mapping name on this profile */
+	FKeyMappingRow* FindKeyMappingRowMutable(const FName InMappingName);
 
-	/** Get all the key mappings associated with the given action name on this profile */
-	const FKeyMappingRow* FindKeyMappingRow(const FName InActionName) const;
+	/** Get all the key mappings associated with the given mapping name on this profile */
+	const FKeyMappingRow* FindKeyMappingRow(const FName InMappingName) const;
 
 	/** A helper function to print out all the current profile settings to the log. */
 	UFUNCTION(BlueprintCallable, Category="Enhanced Input|User Settings")
@@ -253,20 +253,20 @@ public:
 	virtual FString ToString() const;
 	
 	/**
-	 * Returns all FKey's bound to the given Action Name on this profile.
+	 * Returns all FKey's bound to the given mapping Name on this profile.
 	 *
-	 * Returns the number of keys mapped to this action
+	 * Returns the number of keys for the given mapping name
 	 */
 	UFUNCTION(BlueprintCallable, Category="Enhanced Input|User Settings", meta = (ReturnDisplayName = "Number of keys"))
-	virtual int32 GetKeysMappedToAction(const FName ActionName, /*OUT*/ TArray<FKey>& OutKeys) const;
+	virtual int32 GetMappedKeysInRow(const FName MappingName, /*OUT*/ TArray<FKey>& OutKeys) const;
 
 	/**
-	 * Populates the OutMappedActionNames with every action on this profile that has a mapping to the given key.
+	 * Populates the OutMappedMappingNames with every mapping on this profile that has a mapping to the given key.
 	 *
-	 * Returns the number of actions mapped to this key
+	 * Returns the number of mappings to this key
 	 */
-	UFUNCTION(BlueprintCallable, Category="Enhanced Input|User Settings", meta = (ReturnDisplayName = "Number of actions"))
-	virtual int32 GetActionsMappedToKey(const FKey& InKey, /*OUT*/ TArray<FName>& OutMappedActionNames) const;
+	UFUNCTION(BlueprintCallable, Category="Enhanced Input|User Settings", meta = (ReturnDisplayName = "Number of mappings"))
+	virtual int32 GetMappingNamesForKey(const FKey& InKey, /*OUT*/ TArray<FName>& OutMappingNames) const;
 
 	/** Returns a pointer to the player key mapping that fits with the given arguments. Returns null if none exist. */
 	virtual FPlayerKeyMapping* FindKeyMapping(const FMapPlayerKeyArgs& InArgs) const;
@@ -308,7 +308,7 @@ protected:
 	FText DisplayName;
 	
 	/**
-	 * A map of "Action Name" to all key mappings associated with it.
+	 * A map of "Mapping Row Name" to all key mappings associated with it.
 	 * Note: Dirty mappings will be serialized from UEnhancedInputUserSettings::Serialize
 	 */
 	UPROPERTY(BlueprintReadOnly, Transient, EditAnywhere, Category="Enhanced Input|User Settings")
@@ -413,20 +413,21 @@ public:
 	// Remappable keys API
 
 	/**
-	 * Sets the player mapped key for this action on the current key profile.
+	 * Sets the player mapped key on the current key profile.
 	 */
 	UFUNCTION(BlueprintCallable, Category="Enhanced Input|User Settings", meta = (AutoCreateRefTerm = "FailureReason"))
 	virtual void MapPlayerKey(const FMapPlayerKeyArgs& InArgs, FGameplayTagContainer& FailureReason);
 
-	/** Unmap what is currently mapped to the given action in the given slot */
+	/** Unmap what is currently mapped to the given mapping in the given slot */
 	UFUNCTION(BlueprintCallable, Category="Enhanced Input|User Settings", meta = (AutoCreateRefTerm = "FailureReason"))
 	virtual void UnMapPlayerKey(const FMapPlayerKeyArgs& InArgs, FGameplayTagContainer& FailureReason);
 
+	/** Returns a set of all player key mappings for the given mapping name. */
 	UFUNCTION(BlueprintCallable, Category="Enhanced Input|User Settings")
-	virtual const TSet<FPlayerKeyMapping>& FindMappingsForAction(const FName ActionName) const;
+	virtual const TSet<FPlayerKeyMapping>& FindMappingsInRow(const FName MappingName) const;
 
-	/** Returns the current player key mapping for the given action in the given slot */
-	virtual const FPlayerKeyMapping* FindCurrentMappingForSlot(const FName ActionName, const EPlayerMappableKeySlot InSlot) const;
+	/** Returns the current player key mapping for the given row name in the given slot */
+	virtual const FPlayerKeyMapping* FindCurrentMappingForSlot(const FName MappingName, const EPlayerMappableKeySlot InSlot) const;
 	
 	// Modifying key profile
 
@@ -448,9 +449,11 @@ public:
 	UFUNCTION(BlueprintCallable, BlueprintPure, Category="Enhanced Input|User Settings")
 	UEnhancedPlayerMappableKeyProfile* GetCurrentKeyProfile() const;
 
+	/** Get the current key profile that the user has set */
 	template<class T>
-	T* GetCurrentKeyProfile() const
+	inline T* GetCurrentKeyProfile() const
 	{
+		static_assert(TIsDerivedFrom<T, UEnhancedPlayerMappableKeyProfile>::IsDerived, "T must be a UEnhancedPlayerMappableKeyProfile-based type!");
 		return Cast<T>(GetCurrentKeyProfile());
 	}
 
@@ -467,14 +470,16 @@ public:
 	UFUNCTION(BlueprintCallable, BlueprintPure, Category="Enhanced Input|User Settings")
 	virtual UEnhancedPlayerMappableKeyProfile* GetKeyProfileWithIdentifier(const FGameplayTag& ProfileId) const;
 
+	/** Returns the key profile with the given name if one exists. Null if one doesn't exist */
 	template<class T>
-	T* GetKeyProfileWithIdentifier(const FGameplayTag& ProfileId) const
+	inline T* GetKeyProfileWithIdentifier(const FGameplayTag& ProfileId) const
 	{
+		static_assert(TIsDerivedFrom<T, UEnhancedPlayerMappableKeyProfile>::IsDerived, "T must be a UEnhancedPlayerMappableKeyProfile-based type!");
 		return Cast<T>(GetKeyProfileWithIdentifier(ProfileId));
 	}
 
 	// Registering input mapping contexts for access to them from your UI,
-	// even if they are froma plugin
+	// even if they are from a plugin
 
 	/** Fired when a new input mapping context is registered. Useful if you need to update your UI */
 	DECLARE_DYNAMIC_MULTICAST_DELEGATE_OneParam(FMappingContextRegisteredWithSettings, const UInputMappingContext*, IMC);
@@ -485,7 +490,7 @@ public:
 	 * in the context and create an initial Player Mappable Key for every mapping that is marked as mappable.
 	 */
 	UFUNCTION(BlueprintCallable, Category="Enhanced Input|User Settings")
-	bool RegisterInputMappingContext(UInputMappingContext* IMC);
+	virtual bool RegisterInputMappingContext(UInputMappingContext* IMC);
 
 	/** Registers multiple mapping contexts with the settings */
 	UFUNCTION(BlueprintCallable, Category="Enhanced Input|User Settings")
@@ -493,7 +498,7 @@ public:
 
 	/** Removes this mapping context from the registered mapping contexts */
 	UFUNCTION(BlueprintCallable, Category="Enhanced Input|User Settings")
-	bool UnregisterInputMappingContext(const UInputMappingContext* IMC);
+	virtual bool UnregisterInputMappingContext(const UInputMappingContext* IMC);
 
 	/** Removes miltiple mapping contexts from the registered mapping contexts */
 	UFUNCTION(BlueprintCallable, Category="Enhanced Input|User Settings")
