Move the templated GetUserSettings version to just the LP subsystem instead of the interface class

#rnx
#pf skip

diff --git a/Engine/Plugins/EnhancedInput/Source/EnhancedInput/Public/EnhancedInputSubsystemInterface.h b/Engine/Plugins/EnhancedInput/Source/EnhancedInput/Public/EnhancedInputSubsystemInterface.h
--- a/Engine/Plugins/EnhancedInput/Source/EnhancedInput/Public/EnhancedInputSubsystemInterface.h
+++ b/Engine/Plugins/EnhancedInput/Source/EnhancedInput/Public/EnhancedInputSubsystemInterface.h
@@ -76,12 +76,6 @@ public:
 	UFUNCTION(BlueprintCallable, Category="Enhanced Input|User Settings")
 	virtual UEnhancedInputUserSettings* GetUserSettings() const;
 
-	template<class T>
-	inline T* GetUserSettings() const
-	{
-		return Cast<T>(GetUserSettings());
-	}
-
 protected:
 	
 	/**
diff --git a/Engine/Plugins/EnhancedInput/Source/EnhancedInput/Public/EnhancedInputSubsystems.h b/Engine/Plugins/EnhancedInput/Source/EnhancedInput/Public/EnhancedInputSubsystems.h
--- a/Engine/Plugins/EnhancedInput/Source/EnhancedInput/Public/EnhancedInputSubsystems.h
+++ b/Engine/Plugins/EnhancedInput/Source/EnhancedInput/Public/EnhancedInputSubsystems.h
@@ -33,6 +33,12 @@ public:
 	virtual void ControlMappingsRebuiltThisFrame() override;
 	// End IEnhancedInputSubsystemInterface
 
+	template<class UserSettingClass = UEnhancedInputUserSettings>
+	inline UserSettingClass* GetUserSettings() const
+	{
+		return Cast<UserSettingClass>(GetUserSettings());
+	}
+
 	/** A delegate that will be called when control mappings have been rebuilt this frame. */
 	DECLARE_DYNAMIC_MULTICAST_DELEGATE(FOnControlMappingsRebuilt);
 
