Add a gameplay tag container to the player mappable key options to filter down key mappings based on the currently equipped key profile

#rnx

diff --git a/Engine/Plugins/EnhancedInput/Source/EnhancedInput/Private/EnhancedInputSubsystemInterface.cpp b/Engine/Plugins/EnhancedInput/Source/EnhancedInput/Private/EnhancedInputSubsystemInterface.cpp
--- a/Engine/Plugins/EnhancedInput/Source/EnhancedInput/Private/EnhancedInputSubsystemInterface.cpp
+++ b/Engine/Plugins/EnhancedInput/Source/EnhancedInput/Private/EnhancedInputSubsystemInterface.cpp
@@ -11,6 +11,7 @@
 #include "InputMappingContext.h"
 #include "InputMappingQuery.h"
 #include "PlayerMappableInputConfig.h"
+#include "PlayerMappableKeySettings.h"
 
 #include UE_INLINE_GENERATED_CPP_BY_NAME(EnhancedInputSubsystemInterface)
 
@@ -781,6 +782,14 @@ void IEnhancedInputSubsystemInterface::RebuildControlMappings()
 			// True if there were any player mapped keys to this mapping and we are using those instead.
 			bool bIsPlayerMapping = false;
 			
+			const UPlayerMappableKeySettings* KeySettings = Mapping.GetPlayerMappableKeySettings();
+
+			// If this mapping has specified a specific key profile, and the current profile isn't it, then don't add this key mapping
+			if (KeySettings && PlayerKeyProfile && !KeySettings->SupportedKeyProfiles.IsEmpty() && !KeySettings->SupportedKeyProfiles.HasTag(PlayerKeyProfile->GetProfileIdentifer()))
+			{
+				continue;			
+			}
+			
 			// See if there are any player mapped keys to this action
 			if (PlayerKeyProfile && GetDefault<UEnhancedInputDeveloperSettings>()->bEnableUserSettings)
 			{
diff --git a/Engine/Plugins/EnhancedInput/Source/EnhancedInput/Public/PlayerMappableKeySettings.h b/Engine/Plugins/EnhancedInput/Source/EnhancedInput/Public/PlayerMappableKeySettings.h
--- a/Engine/Plugins/EnhancedInput/Source/EnhancedInput/Public/PlayerMappableKeySettings.h
+++ b/Engine/Plugins/EnhancedInput/Source/EnhancedInput/Public/PlayerMappableKeySettings.h
@@ -2,6 +2,8 @@
 
 #pragma once
 
+#include "GameplayTagContainer.h"
+
 #include "PlayerMappableKeySettings.generated.h"
 
 struct FEnhancedActionKeyMapping;
@@ -49,4 +51,12 @@ public:
 	/** The category that this player mapping is in */
 	UPROPERTY(EditAnywhere, BlueprintReadOnly, Category = "Settings")
 	FText DisplayCategory = FText::GetEmpty();
+
+	/** 
+	* If this key mapping should only be added when a specific key profile is equipped, then set those here.
+	* 
+	* If this is empty, then the key mapping will not be filtered out based on the current profile.
+	*/
+	UPROPERTY(EditAnywhere, BlueprintReadOnly, Category = "Settings")
+	FGameplayTagContainer SupportedKeyProfiles;
 };
