Check for a null local player when calling GetPlayerInput on the EI LP subsystem

#rnx

diff --git a/Engine/Plugins/EnhancedInput/Source/EnhancedInput/Private/EnhancedInputSubsystems.cpp b/Engine/Plugins/EnhancedInput/Source/EnhancedInput/Private/EnhancedInputSubsystems.cpp
--- a/Engine/Plugins/EnhancedInput/Source/EnhancedInput/Private/EnhancedInputSubsystems.cpp
+++ b/Engine/Plugins/EnhancedInput/Source/EnhancedInput/Private/EnhancedInputSubsystems.cpp
@@ -36,11 +36,15 @@ void UEnhancedInputLocalPlayerSubsystem::PlayerControllerChanged(APlayerControll
 }
 
 UEnhancedPlayerInput* UEnhancedInputLocalPlayerSubsystem::GetPlayerInput() const
-{	
-	if (APlayerController* PlayerController = GetLocalPlayer()->GetPlayerController(GetWorld()))
+{
+	if (ULocalPlayer* LP = GetLocalPlayer())
 	{
-		return Cast<UEnhancedPlayerInput>(PlayerController->PlayerInput);
+		if (APlayerController* PlayerController = LP->GetPlayerController(GetWorld()))
+		{
+			return Cast<UEnhancedPlayerInput>(PlayerController->PlayerInput);	
+		}
 	}
+	
 	return nullptr;
 }
 
