fix an ensure when binding settings delegates in enhanced input

#rnx
#pf skip

diff --git a/Engine/Plugins/EnhancedInput/Source/EnhancedInput/Private/EnhancedInputSubsystemInterface.cpp b/Engine/Plugins/EnhancedInput/Source/EnhancedInput/Private/EnhancedInputSubsystemInterface.cpp
--- a/Engine/Plugins/EnhancedInput/Source/EnhancedInput/Private/EnhancedInputSubsystemInterface.cpp
+++ b/Engine/Plugins/EnhancedInput/Source/EnhancedInput/Private/EnhancedInputSubsystemInterface.cpp
@@ -58,8 +58,8 @@ void IEnhancedInputSubsystemInterface::BindUserSettingDelegates()
 		return;
 	}
 
-	Settings->OnSettingsChanged.AddDynamic(this, &IEnhancedInputSubsystemInterface::OnUserSettingsChanged);
-	Settings->OnKeyProfileChanged.AddDynamic(this, &IEnhancedInputSubsystemInterface::OnUserKeyProfileChanged);
+	Settings->OnSettingsChanged.AddUniqueDynamic(this, &IEnhancedInputSubsystemInterface::OnUserSettingsChanged);
+	Settings->OnKeyProfileChanged.AddUniqueDynamic(this, &IEnhancedInputSubsystemInterface::OnUserKeyProfileChanged);
 }
 
 void IEnhancedInputSubsystemInterface::OnUserSettingsChanged(UEnhancedInputUserSettings* Settings)
