Allow for input injection based on a player key mapping name as well as a UInputAction pointer. We can do this by keeping track of the UInputAction* associated with each player mapping. Keep track of the owning Platform User ID of each player mappable key profile so that you can access the Local Player without needing to get the outer of the profile, which will be useful for custom profile subclasses.

Rename the templated version of the Get Key Profile template functions to GetKeyProfileAs to avoid compilation errors when using those functions

diff --git a/Engine/Plugins/EnhancedInput/Source/EnhancedInput/Private/EnhancedInputSubsystemInterface.cpp b/Engine/Plugins/EnhancedInput/Source/EnhancedInput/Private/EnhancedInputSubsystemInterface.cpp
--- a/Engine/Plugins/EnhancedInput/Source/EnhancedInput/Private/EnhancedInputSubsystemInterface.cpp
+++ b/Engine/Plugins/EnhancedInput/Source/EnhancedInput/Private/EnhancedInputSubsystemInterface.cpp
@@ -88,6 +88,31 @@ void IEnhancedInputSubsystemInterface::InjectInputVectorForAction(const UInputAc
 	InjectInputForAction(Action, RawValue, Modifiers, Triggers);
 }
 
+void IEnhancedInputSubsystemInterface::InjectInputForPlayerMapping(const FName MappingName, FInputActionValue RawValue, const TArray<UInputModifier*>& Modifiers, const TArray<UInputTrigger*>& Triggers)
+{
+	InjectInputVectorForPlayerMapping(MappingName, RawValue.Get<FVector>(), Modifiers, Triggers);
+}
+
+void IEnhancedInputSubsystemInterface::InjectInputVectorForPlayerMapping(const FName MappingName, FVector Value, const TArray<UInputModifier*>& Modifiers, const TArray<UInputTrigger*>& Triggers)
+{
+	if (const UEnhancedInputUserSettings* UserSettings = GetUserSettings())
+	{
+		if (const UInputAction* Action = UserSettings->FindInputActionForMapping(MappingName))
+		{
+			FInputActionValue RawValue(Action->ValueType, Value);
+			InjectInputForAction(Action, RawValue, Modifiers, Triggers);
+		}
+		else
+		{
+			UE_LOG(LogEnhancedInput, Warning, TEXT("Could not find a Input Action for mapping name '%s'"), *MappingName.ToString());
+		}
+	}
+	else
+	{
+		UE_LOG(LogEnhancedInput, Warning, TEXT("Could not find a valid UEnhancedInputUserSettings object, is it enabled in the project settings?"));
+	}
+}
+
 void IEnhancedInputSubsystemInterface::ClearAllMappings()
 {
 	if (UEnhancedPlayerInput* PlayerInput = GetPlayerInput())
diff --git a/Engine/Plugins/EnhancedInput/Source/EnhancedInput/Private/UserSettings/EnhancedInputUserSettings.cpp b/Engine/Plugins/EnhancedInput/Source/EnhancedInput/Private/UserSettings/EnhancedInputUserSettings.cpp
--- a/Engine/Plugins/EnhancedInput/Source/EnhancedInput/Private/UserSettings/EnhancedInputUserSettings.cpp
+++ b/Engine/Plugins/EnhancedInput/Source/EnhancedInput/Private/UserSettings/EnhancedInputUserSettings.cpp
@@ -5,6 +5,7 @@
 #include "EnhancedPlayerInput.h"
 #include "EnhancedInputModule.h"
 #include "InputMappingContext.h"
+#include "InputAction.h"
 #include "PlayerMappableKeySettings.h"
 #include "EnhancedActionKeyMapping.h"
 #include "EnhancedInputSubsystems.h"
@@ -93,6 +94,7 @@ FPlayerKeyMapping::FPlayerKeyMapping()
 	, DefaultKey(EKeys::Invalid)
 	, CurrentKey(EKeys::Invalid)
 	, HardwareDeviceId(FHardwareDeviceIdentifier::Invalid)
+	, AssociatedInputAction(nullptr)
 {
 }
 
@@ -103,6 +105,7 @@ FPlayerKeyMapping::FPlayerKeyMapping(const FEnhancedActionKeyMapping& OriginalMa
 	, DefaultKey(OriginalMapping.Key)
 	, CurrentKey(EKeys::Invalid)
 	, HardwareDeviceId(FHardwareDeviceIdentifier::Invalid)
+	, AssociatedInputAction(OriginalMapping.Action)
 {
 }
 
@@ -135,12 +138,13 @@ FString FPlayerKeyMapping::ToString() const
 	const UEnum* PlayerMappableEnumClass = StaticEnum<EPlayerMappableKeySlot>();
 	check(PlayerMappableEnumClass);
 	return
-		FString::Printf(TEXT("Action Name: '%s'  Slot: '%s'  Default Key: '%s'  Player Mapped Key: '%s'  HardwareDevice:  '%s'"),
+		FString::Printf(TEXT("Mapping Name: '%s'  Slot: '%s'  Default Key: '%s'  Player Mapped Key: '%s'  HardwareDevice:  '%s'   AssociatedInputAction: '%s'"),
 			*MappingName.ToString(),
 			*PlayerMappableEnumClass->GetNameStringByValue(static_cast<int64>(Slot)),
 			*DefaultKey.ToString(),
 			*CurrentKey.ToString(),
-			*HardwareDeviceId.ToString());
+			*HardwareDeviceId.ToString(),
+			AssociatedInputAction ? *AssociatedInputAction->GetFName().ToString() : TEXT("NULL INPUT ACTION"));
 }
 
 const FName FPlayerKeyMapping::GetMappingName() const
@@ -163,6 +167,11 @@ const FHardwareDeviceIdentifier& FPlayerKeyMapping::GetHardwareDeviceId() const
 	return HardwareDeviceId;
 }
 
+const UInputAction* FPlayerKeyMapping::GetAssociatedInputAction() const
+{
+	return AssociatedInputAction;
+}
+
 uint32 GetTypeHash(const FPlayerKeyMapping& InMapping)
 {
 	uint32 Hash = 0;
@@ -197,6 +206,7 @@ void FPlayerKeyMapping::UpdateOriginalKey(const FEnhancedActionKeyMapping& Origi
 	
 	DefaultKey = OriginalMapping.Key;
 	DisplayName = OriginalMapping.GetDisplayName();
+	AssociatedInputAction = OriginalMapping.Action;
 
 	// If the mapping is the same as the default key, make sure that it is not marked as
 	// dirty so that we don't serialize it
@@ -445,8 +455,11 @@ void UEnhancedInputUserSettings::Initialize(UEnhancedPlayerInput* InPlayerInput)
 	// Create a default key mapping profile in the case where one doesn't exist
 	if (!GetCurrentKeyProfile())
 	{
+		const ULocalPlayer* LP = GetLocalPlayer();
+		
 		FPlayerMappableKeyProfileCreationArgs Args = {};
 		Args.ProfileIdentifier = UE::EnhancedInput::TAG_DefaultProfileIdentifier;
+		Args.UserId = LP ? LP->GetPlatformUserId() : PLATFORMUSERID_NONE;
 		Args.DisplayName = UE::EnhancedInput::DefaultProfileDisplayName;
 		Args.bSetAsCurrentProfile = true;
 		
@@ -792,6 +805,21 @@ const FPlayerKeyMapping* UEnhancedInputUserSettings::FindCurrentMappingForSlot(c
 	return nullptr;
 }
 
+const UInputAction* UEnhancedInputUserSettings::FindInputActionForMapping(const FName MappingName) const
+{
+	const TSet<FPlayerKeyMapping>& AllMappings = FindMappingsInRow(MappingName);
+	// Return the first instance of the input action from the player mappings. All of them will have the same Input Action pointer
+	// as long as they are in the same row.
+	for (const FPlayerKeyMapping& Mapping : AllMappings)
+	{
+		if (const UInputAction* Action = Mapping.GetAssociatedInputAction())
+		{
+			return Action;
+		}
+	}
+	return nullptr;
+}
+
 bool UEnhancedInputUserSettings::SetKeyProfile(const FGameplayTag& InProfileId)
 {
 	UEnhancedPlayerInput* PlayerInput = GetPlayerInput();
@@ -844,6 +872,7 @@ UEnhancedPlayerMappableKeyProfile* UEnhancedInputUserSettings::GetCurrentKeyProf
 FPlayerMappableKeyProfileCreationArgs::FPlayerMappableKeyProfileCreationArgs()
 	: ProfileType(GetDefault<UEnhancedInputDeveloperSettings>()->DefaultPlayerMappableKeyProfileClass.Get())
 	, ProfileIdentifier(FGameplayTag::EmptyTag)
+	, UserId(PLATFORMUSERID_NONE)
 	, DisplayName(FText::GetEmpty())
 	, bSetAsCurrentProfile(true)
 {
@@ -875,6 +904,7 @@ UEnhancedPlayerMappableKeyProfile* UEnhancedInputUserSettings::CreateNewKeyProfi
 		OutProfile = NewObject<UEnhancedPlayerMappableKeyProfile>(/* outer */ this, /* class */ InArgs.ProfileType);
 		OutProfile->ProfileIdentifier = InArgs.ProfileIdentifier;
 		OutProfile->DisplayName = InArgs.DisplayName;
+		OutProfile->OwningUserId = InArgs.UserId;
 		
 		SavedKeyProfiles.Add(InArgs.ProfileIdentifier, OutProfile);
 	}
diff --git a/Engine/Plugins/EnhancedInput/Source/EnhancedInput/Public/EnhancedInputSubsystemInterface.h b/Engine/Plugins/EnhancedInput/Source/EnhancedInput/Public/EnhancedInputSubsystemInterface.h
--- a/Engine/Plugins/EnhancedInput/Source/EnhancedInput/Public/EnhancedInputSubsystemInterface.h
+++ b/Engine/Plugins/EnhancedInput/Source/EnhancedInput/Public/EnhancedInputSubsystemInterface.h
@@ -122,6 +122,30 @@ public:
 	UFUNCTION(BlueprintCallable, Category="Input", meta=(AutoCreateRefTerm="Modifiers,Triggers"))
 	virtual void InjectInputVectorForAction(const UInputAction* Action, FVector Value, const TArray<UInputModifier*>& Modifiers, const TArray<UInputTrigger*>& Triggers);
 
+	/**
+	 * Input simulation via injection. Runs modifiers and triggers delegates as if the input had come through the underlying input system as FKeys.
+	 * Applies action modifiers and triggers on top.
+	 *
+	 * @param MappingName		The name of the player mapping that can be used for look up an associated UInputAction object.
+	 * @param RawValue			The value to set the action to
+	 * @param Modifiers			The modifiers to apply to the injected input.
+	 * @param Triggers			The triggers to apply to the injected input.
+	 */
+	UFUNCTION(BlueprintCallable, Category="Input", meta=(AutoCreateRefTerm="Modifiers,Triggers"))
+	virtual void InjectInputForPlayerMapping(const FName MappingName, FInputActionValue RawValue, const TArray<UInputModifier*>& Modifiers, const TArray<UInputTrigger*>& Triggers);
+	
+	/**
+	 * Input simulation via injection. Runs modifiers and triggers delegates as if the input had come through the underlying input system as FKeys.
+	 * Applies action modifiers and triggers on top.
+	 *
+	 * @param MappingName		The name of the player mapping that can be used for look up an associated UInputAction object.
+	 * @param Value				The value to set the action to (the type will be controlled by the Action)
+	 * @param Modifiers			The modifiers to apply to the injected input.
+	 * @param Triggers			The triggers to apply to the injected input.
+	 */
+	UFUNCTION(BlueprintCallable, Category="Input", meta=(AutoCreateRefTerm="Modifiers,Triggers"))
+	virtual void InjectInputVectorForPlayerMapping(const FName MappingName, FVector Value, const TArray<UInputModifier*>& Modifiers, const TArray<UInputTrigger*>& Triggers);
+	
 	/**
 	 * Remove all applied mapping contexts.
 	 */
diff --git a/Engine/Plugins/EnhancedInput/Source/EnhancedInput/Public/UserSettings/EnhancedInputUserSettings.h b/Engine/Plugins/EnhancedInput/Source/EnhancedInput/Public/UserSettings/EnhancedInputUserSettings.h
--- a/Engine/Plugins/EnhancedInput/Source/EnhancedInput/Public/UserSettings/EnhancedInputUserSettings.h
+++ b/Engine/Plugins/EnhancedInput/Source/EnhancedInput/Public/UserSettings/EnhancedInputUserSettings.h
@@ -14,6 +14,7 @@ class UInputMappingContext;
 class UEnhancedPlayerInput;
 class ULocalPlayer;
 struct FEnhancedActionKeyMapping;
+class UInputAction;
 
 /**
  * The "Slot" that a player mappable key is in.
@@ -126,6 +127,9 @@ public:
 	/** Returns the optional hardware device ID that this mapping is specific to */
 	const FHardwareDeviceIdentifier& GetHardwareDeviceId() const; 
 
+	/** Returns the input actoin asset associated with this player key mapping */
+	const UInputAction* GetAssociatedInputAction() const;
+
 	/** Resets the current mapping to the default one */
 	void ResetToDefault();
 
@@ -175,6 +179,10 @@ protected:
 	/** An optional Hardware Device specifier for this mapping */
 	UPROPERTY(VisibleAnywhere, BlueprintReadOnly, Category="Enhanced Input|User Settings")
 	FHardwareDeviceIdentifier HardwareDeviceId;
+
+	/** The input action associated with this player key mapping */
+	UPROPERTY(Transient, VisibleAnywhere, BlueprintReadOnly, Category="Enhanced Input|User Settings")
+	TObjectPtr<const UInputAction> AssociatedInputAction; 
 };
 
 /**
@@ -302,6 +310,10 @@ protected:
 	/** The ID of this profile. This can be used by each Key Mapping to filter down which profile is required for it be equipped. */
 	UPROPERTY(BlueprintReadOnly, SaveGame, EditAnywhere, Category="Enhanced Input|User Settings")
 	FGameplayTag ProfileIdentifier;
+
+	/** The platform user id of the owning Local Player of this profile. */
+	UPROPERTY(Transient, BlueprintReadOnly, VisibleAnywhere, Category="Enhanced Input|User Settings")
+	FPlatformUserId OwningUserId;
 	
 	/** The localized display name of this profile */
 	UPROPERTY(BlueprintReadWrite, SaveGame, EditAnywhere, Category="Enhanced Input|User Settings")
@@ -329,6 +341,10 @@ struct ENHANCEDINPUT_API FPlayerMappableKeyProfileCreationArgs
 	/** The uniqiue identifier that this profile should have */
 	UPROPERTY(BlueprintReadWrite, EditAnywhere, Category="Enhanced Input|User Settings")
 	FGameplayTag ProfileIdentifier;
+
+	/** The user ID of the ULocalPlayer that this profile is associated with */
+	UPROPERTY(BlueprintReadWrite, EditAnywhere, Category="Enhanced Input|User Settings")
+	FPlatformUserId UserId;
 	
 	/** The display name of this profile */
 	UPROPERTY(BlueprintReadWrite, EditAnywhere, Category="Enhanced Input|User Settings")
@@ -428,6 +444,9 @@ public:
 
 	/** Returns the current player key mapping for the given row name in the given slot */
 	virtual const FPlayerKeyMapping* FindCurrentMappingForSlot(const FName MappingName, const EPlayerMappableKeySlot InSlot) const;
+
+	/** Returns the Input Action associated with the given player mapping name */
+	const UInputAction* FindInputActionForMapping(const FName MappingName) const;
 	
 	// Modifying key profile
 
@@ -451,7 +470,7 @@ public:
 
 	/** Get the current key profile that the user has set */
 	template<class T>
-	inline T* GetCurrentKeyProfile() const
+	inline T* GetCurrentKeyProfileAs() const
 	{
 		static_assert(TIsDerivedFrom<T, UEnhancedPlayerMappableKeyProfile>::IsDerived, "T must be a UEnhancedPlayerMappableKeyProfile-based type!");
 		return Cast<T>(GetCurrentKeyProfile());
@@ -472,7 +491,7 @@ public:
 
 	/** Returns the key profile with the given name if one exists. Null if one doesn't exist */
 	template<class T>
-	inline T* GetKeyProfileWithIdentifier(const FGameplayTag& ProfileId) const
+	inline T* GetKeyProfileWithIdentifierAs(const FGameplayTag& ProfileId) const
 	{
 		static_assert(TIsDerivedFrom<T, UEnhancedPlayerMappableKeyProfile>::IsDerived, "T must be a UEnhancedPlayerMappableKeyProfile-based type!");
 		return Cast<T>(GetKeyProfileWithIdentifier(ProfileId));
