PR #9693: Hold triggers don't reset on mapping context changes (Contributed by [REDACTED_PERSON])

This fixes an issue when comparing enhanced action mappings during RebuildControlMappings because we had made a deep copy of the modifiers and triggers arrays, so any time we compare a new instance to old instance it would return false no matter what.

diff --git a/Engine/Plugins/EnhancedInput/Source/EnhancedInput/Private/EnhancedInputSubsystemInterface.cpp b/Engine/Plugins/EnhancedInput/Source/EnhancedInput/Private/EnhancedInputSubsystemInterface.cpp
--- a/Engine/Plugins/EnhancedInput/Source/EnhancedInput/Private/EnhancedInputSubsystemInterface.cpp
+++ b/Engine/Plugins/EnhancedInput/Source/EnhancedInput/Private/EnhancedInputSubsystemInterface.cpp
@@ -934,7 +934,16 @@ void IEnhancedInputSubsystemInterface::RebuildControlMappings()
 			}
 
 			// Retain old mapping trigger/modifier state for identical key -> action mappings.
-			TArray<FEnhancedActionKeyMapping>::SizeType Idx = OldMappings.IndexOfByPredicate([&Mapping](const FEnhancedActionKeyMapping& Other) {return Mapping == Other; });
+			TArray<FEnhancedActionKeyMapping>::SizeType Idx = OldMappings.IndexOfByPredicate(
+				[&Mapping](const FEnhancedActionKeyMapping& Other)
+				{
+					// Use Equals() to ignore Triggers' values. We want to keep their values from before remapping to
+					// prevent resets. Otherwise, triggers like UInputTriggerPressed re-trigger when their value is
+					// reset to 0; and time counting triggers, like UInputTriggerHold, restart their time.
+					// But don't ignore Modifier and Trigger types and their order in the comparison. If we did, we'd
+					// replace new mappings for old ones with different Trigger and Modifier settings.
+					return Mapping.Equals(Other);
+				});
 			if (Idx != INDEX_NONE)
 			{
 				Mapping = MoveTemp(OldMappings[Idx]);
diff --git a/Engine/Plugins/EnhancedInput/Source/EnhancedInput/Public/EnhancedActionKeyMapping.h b/Engine/Plugins/EnhancedInput/Source/EnhancedInput/Public/EnhancedActionKeyMapping.h
--- a/Engine/Plugins/EnhancedInput/Source/EnhancedInput/Public/EnhancedActionKeyMapping.h
+++ b/Engine/Plugins/EnhancedInput/Source/EnhancedInput/Public/EnhancedActionKeyMapping.h
@@ -109,6 +109,7 @@ struct ENHANCEDINPUT_API FEnhancedActionKeyMapping
 	EDataValidationResult IsDataValid(TArray<FText>& ValidationErrors);
 #endif
 	
+	/** Identical comparison, including Triggers and Modifiers current inner values. */
 	bool operator==(const FEnhancedActionKeyMapping& Other) const;
 
 #if WITH_EDITORONLY_DATA
@@ -183,6 +184,67 @@ public:
 		return Cast<T>(GetPlayerMappableKeySettings());
 	}
 
+	/**
+	 * If the template bIgnoreModifierAndTriggerValues is true, compare to Other ignoring
+	 * different trigger states, like pending activation time, but only accept
+	 * both as equal if the Trigger types are the same and in the same order.
+	 */
+	template<bool bIgnoreModifierAndTriggerValues = true>
+	bool Equals(const FEnhancedActionKeyMapping& Other) const
+	{
+		if constexpr (bIgnoreModifierAndTriggerValues)
+		{
+			return (Action == Other.Action &&
+					Key == Other.Key &&
+					CompareByObjectTypes(Modifiers, Other.Modifiers) &&
+					CompareByObjectTypes(Triggers, Other.Triggers));
+		}
+		else
+		{
+			return *this == Other;
+		}
+	}
+	
+	/**
+	 * Compares if two TArray of UObjects contain the same number and types of
+	 * objects, but doesn't compare their values.
+	 *
+	 * This is needed because TArray::operator== returns false when the objects'
+	 * inner values differ. And for keeping old Trigger states, we need their
+	 * comparison to ignore their current values.
+	 */
+	template<typename T>
+	static bool CompareByObjectTypes(const TArray<TObjectPtr<T>>& A, const TArray<TObjectPtr<T>>& B)
+	{
+		if (A.Num() != B.Num())
+		{
+			return false;
+		}
+
+		for (int32 Idx = 0; Idx < A.Num(); ++Idx)
+		{
+			const T* ObjectA = A[Idx];
+			const T* ObjectB = B[Idx];
+
+			if ((ObjectA == nullptr) != (ObjectB == nullptr))
+			{
+				// One is nullptr, the other isn't
+				return false;
+			}
+			if (!ObjectA)
+			{
+				// Both are nullptr. Consider that as same type.
+				continue;
+			}
+			if (ObjectA->GetClass() != ObjectB->GetClass())
+			{
+				return false;
+			}
+		}
+
+		return true;
+	}
+
 };
 
 PRAGMA_ENABLE_DEPRECATION_WARNINGS
