Make the pointer to the owning UPlayerInput object a weak pointer because that may change when the player controller changes. Instead, update that pointer whenever we get a new player controller to ensure that it is up to date.

#rnx

diff --git a/Engine/Plugins/EnhancedInput/Source/EnhancedInput/Private/EnhancedInputSubsystems.cpp b/Engine/Plugins/EnhancedInput/Source/EnhancedInput/Private/EnhancedInputSubsystems.cpp
--- a/Engine/Plugins/EnhancedInput/Source/EnhancedInput/Private/EnhancedInputSubsystems.cpp
+++ b/Engine/Plugins/EnhancedInput/Source/EnhancedInput/Private/EnhancedInputSubsystems.cpp
@@ -55,9 +55,11 @@ UEnhancedInputUserSettings* UEnhancedInputLocalPlayerSubsystem::GetUserSettings(
 
 void UEnhancedInputLocalPlayerSubsystem::InitalizeUserSettings()
 {
-	// We don't want to be re-creating any user settings, they should share the lifetime with the player input object
-	if (!ensureMsgf(!UserSettings, TEXT("Attempting to initalize the User Settings, but they have already been created!")))
+	// If the user settings are already created then we just want to update it's pointer to the 
+	// player input object, as that may have changed with a new player controller.
+	if (UserSettings)
 	{
+		UserSettings->Initialize(GetPlayerInput());
 		return;
 	}
 
diff --git a/Engine/Plugins/EnhancedInput/Source/EnhancedInput/Private/UserSettings/EnhancedInputUserSettings.cpp b/Engine/Plugins/EnhancedInput/Source/EnhancedInput/Private/UserSettings/EnhancedInputUserSettings.cpp
--- a/Engine/Plugins/EnhancedInput/Source/EnhancedInput/Private/UserSettings/EnhancedInputUserSettings.cpp
+++ b/Engine/Plugins/EnhancedInput/Source/EnhancedInput/Private/UserSettings/EnhancedInputUserSettings.cpp
@@ -656,7 +656,7 @@ void UEnhancedInputUserSettings::Serialize(FArchive& Ar)
 
 UEnhancedPlayerInput* UEnhancedInputUserSettings::GetPlayerInput() const
 {
-	return OwningPlayerInput;
+	return OwningPlayerInput.Get();
 }
 
 APlayerController* UEnhancedInputUserSettings::GetPlayerController() const
diff --git a/Engine/Plugins/EnhancedInput/Source/EnhancedInput/Public/UserSettings/EnhancedInputUserSettings.h b/Engine/Plugins/EnhancedInput/Source/EnhancedInput/Public/UserSettings/EnhancedInputUserSettings.h
--- a/Engine/Plugins/EnhancedInput/Source/EnhancedInput/Public/UserSettings/EnhancedInputUserSettings.h
+++ b/Engine/Plugins/EnhancedInput/Source/EnhancedInput/Public/UserSettings/EnhancedInputUserSettings.h
@@ -544,7 +544,7 @@ protected:
 	
 	/** The owning Player Input object of these settings */
 	UPROPERTY(Transient)
-	TObjectPtr<UEnhancedPlayerInput> OwningPlayerInput;
+	TWeakObjectPtr<UEnhancedPlayerInput> OwningPlayerInput;
 	
 	/**
 	 * Set of currently registered input mapping contexts that may not be currently
