Add some querying options to the settings class that will allow us to filter down key mappings based on the current hardware device primary types and their supported features.

diff --git a/Engine/Plugins/EnhancedInput/Source/EnhancedInput/Private/EnhancedInputSubsystemInterface.cpp b/Engine/Plugins/EnhancedInput/Source/EnhancedInput/Private/EnhancedInputSubsystemInterface.cpp
--- a/Engine/Plugins/EnhancedInput/Source/EnhancedInput/Private/EnhancedInputSubsystemInterface.cpp
+++ b/Engine/Plugins/EnhancedInput/Source/EnhancedInput/Private/EnhancedInputSubsystemInterface.cpp
@@ -803,9 +803,6 @@ void IEnhancedInputSubsystemInterface::RebuildControlMappings()
 		{
 			// Clear out mappings from the previous iteration
 			MappedKeysToActionName.Reset();
-
-			// True if there were any player mapped keys to this mapping and we are using those instead.
-			bool bIsPlayerMapping = false;
 			
 			const UPlayerMappableKeySettings* KeySettings = Mapping.GetPlayerMappableKeySettings();
 
@@ -818,15 +815,16 @@ void IEnhancedInputSubsystemInterface::RebuildControlMappings()
 			// See if there are any player mapped keys to this action
 			if (PlayerKeyProfile && GetDefault<UEnhancedInputDeveloperSettings>()->bEnableUserSettings)
 			{
-				PlayerKeyProfile->GetMappedKeysInRow(Mapping.GetMappingName(), MappedKeysToActionName);
-				bIsPlayerMapping = !MappedKeysToActionName.IsEmpty();
+				PlayerKeyProfile->GetPlayerMappedKeysForRebuildControlMappings(Mapping, MappedKeysToActionName);
 			}
 			else if (const TMap<FPlayerMappableKeySlot, FKey>* PlayerMappedKeySlots = PlayerMappedSettings.Find(Mapping.GetMappingName()))
 			{
 				PlayerMappedKeySlots->GenerateValueArray(MappedKeysToActionName);
-				bIsPlayerMapping = !MappedKeysToActionName.IsEmpty();
 			}
 
+			// True if there were any player mapped keys to this mapping and we are using those instead.
+			const bool bIsPlayerMapping = !MappedKeysToActionName.IsEmpty();
+
 			// If there aren't, then just use the default mapping for this action
 			if (!bIsPlayerMapping)
 			{
diff --git a/Engine/Plugins/EnhancedInput/Source/EnhancedInput/Private/UserSettings/EnhancedInputUserSettings.cpp b/Engine/Plugins/EnhancedInput/Source/EnhancedInput/Private/UserSettings/EnhancedInputUserSettings.cpp
--- a/Engine/Plugins/EnhancedInput/Source/EnhancedInput/Private/UserSettings/EnhancedInputUserSettings.cpp
+++ b/Engine/Plugins/EnhancedInput/Source/EnhancedInput/Private/UserSettings/EnhancedInputUserSettings.cpp
@@ -341,11 +341,109 @@ int32 UEnhancedPlayerMappableKeyProfile::GetMappedKeysInRow(const FName MappingN
 	}
 	else
 	{
-		UE_LOG(LogEnhancedInput, Warning, TEXT("Player Mappable Key Profile '%s' doesn't have any mappings for action '%s'"), *ProfileIdentifier.ToString(), *MappingName.ToString());
+		UE_LOG(LogEnhancedInput, Warning, TEXT("Player Mappable Key Profile '%s' doesn't have any mappings for '%s'"), *ProfileIdentifier.ToString(), *MappingName.ToString());
 	}
 	return OutKeys.Num();
 }
 
+int32 UEnhancedPlayerMappableKeyProfile::GetPlayerMappedKeysForRebuildControlMappings(const FEnhancedActionKeyMapping& DefaultMapping, TArray<FKey>& OutKeys) const
+{
+	FPlayerMappableKeyQueryOptions QueryOptions = {};
+
+	QueryOptions.MappingName = DefaultMapping.GetMappingName();
+	QueryOptions.KeyToMatch = DefaultMapping.Key;
+	QueryOptions.bMatchBasicKeyTypes = true;
+	QueryOptions.bMatchKeyAxisType = false;
+
+	// Note for subclasses, if you want to filter based on the current input device then you can
+	// do this!
+	// if (const UInputDeviceSubsystem* DeviceSubsystem = UInputDeviceSubsystem::Get())
+	// {
+	// 	FHardwareDeviceIdentifier MostRecentDevice = DeviceSubsystem->GetMostRecentlyUsedHardwareDevice(OwningUserId);
+	// 	QueryOptions.RequiredDeviceType = MostRecentDevice.PrimaryDeviceType;
+	// 	QueryOptions.RequiredDeviceFlags = MostRecentDevice.SupportedFeaturesMask;
+	// }
+	
+	return QueryPlayerMappedKeys(QueryOptions, OutKeys);
+}
+
+int32 UEnhancedPlayerMappableKeyProfile::QueryPlayerMappedKeys(const FPlayerMappableKeyQueryOptions& Options, TArray<FKey>& OutKeys) const
+{
+	OutKeys.Reset();
+
+	if (const FKeyMappingRow* MappingRow = FindKeyMappingRow(Options.MappingName))
+	{
+		for (const FPlayerKeyMapping& PlayerMapping : MappingRow->Mappings)
+		{
+			if (DoesMappingPassQueryOptions(PlayerMapping, Options))
+			{
+				OutKeys.Add(PlayerMapping.GetCurrentKey());
+			}
+		}
+	}
+	else
+	{
+		UE_LOG(LogEnhancedInput, Warning, TEXT("Player Mappable Key Profile '%s' doesn't have any mappings for '%s'"), *ProfileIdentifier.ToString(), *Options.MappingName.ToString());
+	}
+
+	return OutKeys.Num();
+}
+
+bool UEnhancedPlayerMappableKeyProfile::DoesMappingPassQueryOptions(const FPlayerKeyMapping& PlayerMapping, const FPlayerMappableKeyQueryOptions& Options) const
+{
+	if (Options.KeyToMatch.IsValid())
+	{
+		const FKey& A = Options.KeyToMatch;
+		const FKey& B = PlayerMapping.GetCurrentKey();
+	
+		// Ensure that the player mapped key matches the one set in the key profile
+		if (Options.bMatchBasicKeyTypes)
+		{
+			const bool bKeyTypesMatch = 
+				A.IsGamepadKey() == B.IsGamepadKey() &&
+				A.IsTouch() == B.IsTouch() &&
+				A.IsGesture() == B.IsGesture();
+
+			if (!bKeyTypesMatch)
+			{
+				return false;
+			}
+		}
+
+		if (Options.bMatchKeyAxisType)
+		{
+			const bool bKeyAxisMatch =
+				A.IsAxis1D() == B.IsAxis1D() &&
+				A.IsAxis2D() == B.IsAxis2D() &&
+				A.IsAxis3D() == B.IsAxis3D();
+
+			if (!bKeyAxisMatch)
+			{
+				return false;
+			}
+		}
+	}
+
+	// Match hardware device info per mapping
+	const FHardwareDeviceIdentifier& PlayerMappingDevice = PlayerMapping.GetHardwareDeviceId();
+	
+	// Filter mappings based on their primary hardware device type
+	if (Options.RequiredDeviceType != EHardwareDevicePrimaryType::Unspecified &&
+		Options.RequiredDeviceType != PlayerMappingDevice.PrimaryDeviceType)
+	{
+		return false;
+	}
+
+	// Filter mappings based on their hardware device's supported features
+	if (Options.RequiredDeviceFlags != EHardwareDeviceSupportedFeatures::Type::Unspecified &&
+		!PlayerMappingDevice.HasAllSupportedFeatures(static_cast<EHardwareDeviceSupportedFeatures::Type>(Options.RequiredDeviceFlags)))
+	{
+		return false;
+	}
+
+	return true;
+}
+
 int32 UEnhancedPlayerMappableKeyProfile::GetMappingNamesForKey(const FKey& InKey, TArray<FName>& OutMappingNames) const
 {
 	OutMappingNames.Reset();
@@ -1070,4 +1168,13 @@ bool UEnhancedInputUserSettings::IsMappingContextRegistered(const UInputMappingC
 	return RegisteredMappingContexts.Contains(IMC);
 }
 
+FPlayerMappableKeyQueryOptions::FPlayerMappableKeyQueryOptions()
+	: MappingName(NAME_None)
+	, KeyToMatch(EKeys::Invalid)
+	, bMatchBasicKeyTypes(true)
+	, bMatchKeyAxisType(false)
+	, RequiredDeviceType(EHardwareDevicePrimaryType::Unspecified)
+	, RequiredDeviceFlags(EHardwareDeviceSupportedFeatures::Type::Unspecified)
+{ }
+
 #undef LOCTEXT_NAMESPACE
\ No newline at end of file
diff --git a/Engine/Plugins/EnhancedInput/Source/EnhancedInput/Public/UserSettings/EnhancedInputUserSettings.h b/Engine/Plugins/EnhancedInput/Source/EnhancedInput/Public/UserSettings/EnhancedInputUserSettings.h
--- a/Engine/Plugins/EnhancedInput/Source/EnhancedInput/Public/UserSettings/EnhancedInputUserSettings.h
+++ b/Engine/Plugins/EnhancedInput/Source/EnhancedInput/Public/UserSettings/EnhancedInputUserSettings.h
@@ -203,6 +203,41 @@ struct ENHANCEDINPUT_API FKeyMappingRow
 	bool HasAnyMappings() const;
 };
 
+/** 
+* Options when querying what keys are mapped to a specific action on the player mappable
+* key profile. 
+*/
+USTRUCT(BlueprintType)
+struct ENHANCEDINPUT_API FPlayerMappableKeyQueryOptions
+{
+	GENERATED_BODY()
+
+	FPlayerMappableKeyQueryOptions();
+
+	/** The mapping name to search for */
+	UPROPERTY(BlueprintReadWrite, EditAnywhere, Category="Query Options")
+	FName MappingName;
+	
+	/** If specified, then this key will be used to match against when checking if the key types and axis are the same. */
+	UPROPERTY(BlueprintReadWrite, EditAnywhere, Category="Query Options")
+	FKey KeyToMatch;
+	
+	/** If true, then only keys that have the same value for IsGamepadKey, IsTouch, and IsGesture will be included in the results of this query */
+	UPROPERTY(BlueprintReadWrite, EditAnywhere, Category="Query Options")
+	uint8 bMatchBasicKeyTypes : 1;
+
+	/** If true, then only keys that have the same value of IsAxis1D, IsAxis2D, and IsAxis3D will be included in the results of this query */
+	UPROPERTY(BlueprintReadWrite, EditAnywhere, Category="Query Options")
+	uint8 bMatchKeyAxisType : 1;
+	
+	/** If set, then only player mappings whose hardware device identifier that has the same primary input device type will be included in the results of this query */
+	UPROPERTY(BlueprintReadWrite, EditAnywhere, Category="Query Options")
+	EHardwareDevicePrimaryType RequiredDeviceType;
+
+	/** If set, then only player mappings whose Hardware Device Identifier that has the same flags as this will be included in the results */
+	UPROPERTY(BlueprintReadWrite, EditAnywhere, Category = "Query Options", meta=(Bitmask, BitmaskEnum="/Script/Engine.EHardwareDeviceSupportedFeatures"))
+	int32 RequiredDeviceFlags;
+};
 
 /** Represents one "Profile" that a user can have for their player mappable keys */
 UCLASS(BlueprintType)
@@ -268,6 +303,27 @@ public:
 	UFUNCTION(BlueprintCallable, Category="Enhanced Input|User Settings", meta = (ReturnDisplayName = "Number of keys"))
 	virtual int32 GetMappedKeysInRow(const FName MappingName, /*OUT*/ TArray<FKey>& OutKeys) const;
 
+	/**
+	 * Called during IEnhancedInputSubsystemInterface::RebuildControlMappings, this provides your user settings
+	 * the oppurtunity to decide what key mappings get added for a given FEnhancedActionKeyMapping, and applied to the player.
+	 *
+	 * Override this to change what query options are being used during the controls being rebuilt of enhanced input.
+	 *
+	 * If the given Out Array is empty, then the Default Mapping will be applied to the player, but no player mapped keys will be.
+	 */
+	virtual int32 GetPlayerMappedKeysForRebuildControlMappings(const FEnhancedActionKeyMapping& DefaultMapping, /*OUT*/ TArray<FKey>& OutKeys) const;
+
+	/** 
+	* Populates the OutKeys array with any player mapped FKeys for the given default mapping. 
+	* 
+	* This is what IEnhancedInputSubsystemInterface::RebuildControlMappings calls to determine
+	* what keys should actually be applied when building the control mappings.
+	*
+	* Returns the number of player mapped keys to the given Default Mapping
+	*/
+	UFUNCTION(BlueprintCallable, Category="Enhanced Input|User Settings", meta = (ReturnDisplayName = "Number of mappings"))
+	virtual int32 QueryPlayerMappedKeys(const FPlayerMappableKeyQueryOptions& Options, /*OUT*/ TArray<FKey>& OutKeys) const;
+
 	/**
 	 * Populates the OutMappedMappingNames with every mapping on this profile that has a mapping to the given key.
 	 *
@@ -290,6 +346,11 @@ public:
 	
 protected:
 
+	/**
+	 * Returns true if the given player key mapping passes the query filter provided. 
+	 */
+	virtual bool DoesMappingPassQueryOptions(const FPlayerKeyMapping& PlayerMapping, const FPlayerMappableKeyQueryOptions& Options) const;
+
 	/**
 	 * Equips the current key profile. This will always be called after the previous key profile's UnEquip function.
 	*  Make any changes to the given Enhanced Player input object that may be necessary for
