Make DoesMappingPassQueryOptions a public function so that UI generating code can use it

#Rnx
#pf skip

diff --git a/Engine/Plugins/EnhancedInput/Source/EnhancedInput/Public/UserSettings/EnhancedInputUserSettings.h b/Engine/Plugins/EnhancedInput/Source/EnhancedInput/Public/UserSettings/EnhancedInputUserSettings.h
--- a/Engine/Plugins/EnhancedInput/Source/EnhancedInput/Public/UserSettings/EnhancedInputUserSettings.h
+++ b/Engine/Plugins/EnhancedInput/Source/EnhancedInput/Public/UserSettings/EnhancedInputUserSettings.h
@@ -344,13 +344,14 @@ public:
 	UFUNCTION(BlueprintCallable, Category="Enhanced Input|User Settings")
 	virtual void ResetToDefault();
 	
-protected:
-
 	/**
 	 * Returns true if the given player key mapping passes the query filter provided. 
 	 */
+	UFUNCTION(BlueprintCallable, Category="Enhanced Input|User Settings")
 	virtual bool DoesMappingPassQueryOptions(const FPlayerKeyMapping& PlayerMapping, const FPlayerMappableKeyQueryOptions& Options) const;
 
+protected:
+	
 	/**
 	 * Equips the current key profile. This will always be called after the previous key profile's UnEquip function.
 	*  Make any changes to the given Enhanced Player input object that may be necessary for
