Fix for not being able to change player mappable key settings due to some incorrect post load logic

#rnx

diff --git a/Engine/Plugins/EnhancedInput/Source/EnhancedInput/Private/InputMappingContext.cpp b/Engine/Plugins/EnhancedInput/Source/EnhancedInput/Private/InputMappingContext.cpp
--- a/Engine/Plugins/EnhancedInput/Source/EnhancedInput/Private/InputMappingContext.cpp
+++ b/Engine/Plugins/EnhancedInput/Source/EnhancedInput/Private/InputMappingContext.cpp
@@ -10,6 +10,11 @@
 
 #define LOCTEXT_NAMESPACE "InputMappingContext"
 
+namespace UE::EnhancedInput
+{
+	static const FName PlayerMappableOptionsHaveBeenUpgradedname = TEXT("__EnhancedInput_Internal_HasBeenUpgraded__");
+}
+
 #if WITH_EDITOR
 EDataValidationResult UInputMappingContext::IsDataValid(TArray<FText>& ValidationErrors)
 {
@@ -32,7 +37,7 @@ void UInputMappingContext::PostLoad()
 	for (FEnhancedActionKeyMapping& Mapping : Mappings)
 	{
 		// If this was player mappable, then set the behavior to override the action and populate it with the same data
-		if (Mapping.bIsPlayerMappable)
+		if (Mapping.bIsPlayerMappable && Mapping.PlayerMappableOptions.Name != UE::EnhancedInput::PlayerMappableOptionsHaveBeenUpgradedname)
 		{
 			Mapping.SettingBehavior = EPlayerMappableKeySettingBehaviors::OverrideSettings;
 
@@ -43,10 +48,9 @@ void UInputMappingContext::PostLoad()
 				EObjectFlags MaskedOuterFlags = RF_Public | RF_Transactional;
 				Mapping.PlayerMappableKeySettings = NewObject<UPlayerMappableKeySettings>(this, UPlayerMappableKeySettings::StaticClass(), *NewObjectName, MaskedOuterFlags, nullptr);
 			}
-		}
-		
-		if (Mapping.PlayerMappableKeySettings)
-		{
+
+			check(Mapping.PlayerMappableKeySettings);
+
 			Mapping.PlayerMappableKeySettings->Metadata = Mapping.PlayerMappableOptions.Metadata;
 
 			// If the name was set already, use that
@@ -60,10 +64,14 @@ void UInputMappingContext::PostLoad()
 				FName UniqueName = MakeUniqueObjectName(this, UInputMappingContext::StaticClass(), Mapping.Action.GetFName());
 				Mapping.PlayerMappableKeySettings->Name = UniqueName;
 			}
-			
+
 			Mapping.PlayerMappableKeySettings->DisplayName = Mapping.PlayerMappableOptions.DisplayName;
 			Mapping.PlayerMappableKeySettings->DisplayCategory = Mapping.PlayerMappableOptions.DisplayCategory;
-		}
+
+			// Set the name of the old player mappable options so that we know it has already been upgraded, and we don't do it again
+			// on the next post load.
+			Mapping.PlayerMappableOptions.Name = UE::EnhancedInput::PlayerMappableOptionsHaveBeenUpgradedname;
+		}		
 	}
 	
 	PRAGMA_ENABLE_DEPRECATION_WARNINGS
