Unshelved from pending changelist '24645662':

Add a display category FText to the Player Mappable keys to support all of the old features

#rnx

diff --git a/Engine/Plugins/EnhancedInput/Source/EnhancedInput/Private/EnhancedActionKeyMapping.cpp b/Engine/Plugins/EnhancedInput/Source/EnhancedInput/Private/EnhancedActionKeyMapping.cpp
--- a/Engine/Plugins/EnhancedInput/Source/EnhancedInput/Private/EnhancedActionKeyMapping.cpp
+++ b/Engine/Plugins/EnhancedInput/Source/EnhancedInput/Private/EnhancedActionKeyMapping.cpp
@@ -66,6 +66,15 @@ const FText& FEnhancedActionKeyMapping::GetDisplayName() const
 	return FText::GetEmpty();
 }
 
+const FText& FEnhancedActionKeyMapping::GetDisplayCategory() const
+{
+	if (UPlayerMappableKeySettings* MappableKeySettings = GetPlayerMappableKeySettings())
+	{
+		return MappableKeySettings->DisplayCategory;
+	}
+	return FText::GetEmpty();
+}
+
 bool FEnhancedActionKeyMapping::IsPlayerMappable() const
 {
 	return GetPlayerMappableKeySettings() != nullptr;
diff --git a/Engine/Plugins/EnhancedInput/Source/EnhancedInput/Private/UserSettings/EnhancedInputUserSettings.cpp b/Engine/Plugins/EnhancedInput/Source/EnhancedInput/Private/UserSettings/EnhancedInputUserSettings.cpp
--- a/Engine/Plugins/EnhancedInput/Source/EnhancedInput/Private/UserSettings/EnhancedInputUserSettings.cpp
+++ b/Engine/Plugins/EnhancedInput/Source/EnhancedInput/Private/UserSettings/EnhancedInputUserSettings.cpp
@@ -90,6 +90,7 @@ FMapPlayerKeyArgs::FMapPlayerKeyArgs()
 FPlayerKeyMapping::FPlayerKeyMapping()
 	: MappingName(NAME_None)
 	, DisplayName(FText::GetEmpty())
+	, DisplayCategory(FText::GetEmpty())
 	, Slot(EPlayerMappableKeySlot::Unspecified)
 	, DefaultKey(EKeys::Invalid)
 	, CurrentKey(EKeys::Invalid)
@@ -101,6 +102,7 @@ FPlayerKeyMapping::FPlayerKeyMapping()
 FPlayerKeyMapping::FPlayerKeyMapping(const FEnhancedActionKeyMapping& OriginalMapping, EPlayerMappableKeySlot InSlot /* = EPlayerMappableKeySlot::Unspecified */)
 	: MappingName(OriginalMapping.GetMappingName())
 	, DisplayName(OriginalMapping.GetDisplayName())
+	, DisplayCategory(OriginalMapping.GetDisplayCategory())
 	, Slot(InSlot)
 	, DefaultKey(OriginalMapping.Key)
 	, CurrentKey(EKeys::Invalid)
@@ -157,6 +159,11 @@ const FText& FPlayerKeyMapping::GetDisplayName() const
 	return DisplayName;
 }
 
+const FText& FPlayerKeyMapping::GetDisplayCategory() const
+{
+	return DisplayCategory;
+}
+
 EPlayerMappableKeySlot FPlayerKeyMapping::GetSlot() const
 {
 	return Slot;
@@ -348,6 +355,12 @@ int32 UEnhancedPlayerMappableKeyProfile::GetMappedKeysInRow(const FName MappingN
 
 int32 UEnhancedPlayerMappableKeyProfile::GetPlayerMappedKeysForRebuildControlMappings(const FEnhancedActionKeyMapping& DefaultMapping, TArray<FKey>& OutKeys) const
 {
+	// Don't bother trying to query for keys on a non-player mappable action mapping
+	if (!DefaultMapping.IsPlayerMappable())
+	{
+		return 0;
+	}
+	
 	FPlayerMappableKeyQueryOptions QueryOptions = {};
 
 	QueryOptions.MappingName = DefaultMapping.GetMappingName();
diff --git a/Engine/Plugins/EnhancedInput/Source/EnhancedInput/Public/EnhancedActionKeyMapping.h b/Engine/Plugins/EnhancedInput/Source/EnhancedInput/Public/EnhancedActionKeyMapping.h
--- a/Engine/Plugins/EnhancedInput/Source/EnhancedInput/Public/EnhancedActionKeyMapping.h
+++ b/Engine/Plugins/EnhancedInput/Source/EnhancedInput/Public/EnhancedActionKeyMapping.h
@@ -98,8 +98,12 @@ struct ENHANCEDINPUT_API FEnhancedActionKeyMapping
 	 */
 	FName GetMappingName() const;
 
+	/** The localized display name of this key mapping */
 	const FText& GetDisplayName() const;
 
+	/** The localized display category of this key mapping */
+	const FText& GetDisplayCategory() const;
+
 	/**
 	* Returns true if this Action Key Mapping either holds a Player Mappable Key Settings or is set bIsPlayerMappable.
 	*/
diff --git a/Engine/Plugins/EnhancedInput/Source/EnhancedInput/Public/UserSettings/EnhancedInputUserSettings.h b/Engine/Plugins/EnhancedInput/Source/EnhancedInput/Public/UserSettings/EnhancedInputUserSettings.h
--- a/Engine/Plugins/EnhancedInput/Source/EnhancedInput/Public/UserSettings/EnhancedInputUserSettings.h
+++ b/Engine/Plugins/EnhancedInput/Source/EnhancedInput/Public/UserSettings/EnhancedInputUserSettings.h
@@ -121,6 +121,9 @@ public:
 	/** The localized display name to use for this mapping */
 	const FText& GetDisplayName() const;
 
+	/** The localized display name for the category to use for this mapping */
+	const FText& GetDisplayCategory() const;
+
 	/** Returns what player mappable slot this mapping is in */
 	EPlayerMappableKeySlot GetSlot() const;
 
@@ -160,6 +163,10 @@ protected:
 	UPROPERTY(Transient, VisibleAnywhere, BlueprintReadOnly, Category="Enhanced Input|User Settings")
 	FText DisplayName;
 
+	/** Localized display category of this mapping */
+	UPROPERTY(Transient, VisibleAnywhere, BlueprintReadOnly, Category="Enhanced Input|User Settings")
+	FText DisplayCategory;
+
 	/** What slot this key is mapped to */
 	UPROPERTY(VisibleAnywhere, BlueprintReadOnly, Category="Enhanced Input|User Settings")
 	EPlayerMappableKeySlot Slot;
