Add an option to Query keys based on their slot

#rnx
#pf trivial

diff --git a/Engine/Plugins/EnhancedInput/Source/EnhancedInput/Private/UserSettings/EnhancedInputUserSettings.cpp b/Engine/Plugins/EnhancedInput/Source/EnhancedInput/Private/UserSettings/EnhancedInputUserSettings.cpp
--- a/Engine/Plugins/EnhancedInput/Source/EnhancedInput/Private/UserSettings/EnhancedInputUserSettings.cpp
+++ b/Engine/Plugins/EnhancedInput/Source/EnhancedInput/Private/UserSettings/EnhancedInputUserSettings.cpp
@@ -437,6 +437,14 @@ bool UEnhancedPlayerMappableKeyProfile::DoesMappingPassQueryOptions(const FPlaye
 		}
 	}
 
+	// If you have requested a specific key slot, then filter down by that.
+	// Do not filter if the slot to match is Unspecified.
+	if (Options.SlotToMatch != EPlayerMappableKeySlot::Unspecified &&
+		Options.SlotToMatch != PlayerMapping.GetSlot())
+	{
+		return false;
+	}
+
 	// Match hardware device info per mapping
 	const FHardwareDeviceIdentifier& PlayerMappingDevice = PlayerMapping.GetHardwareDeviceId();
 	
@@ -1184,6 +1192,7 @@ bool UEnhancedInputUserSettings::IsMappingContextRegistered(const UInputMappingC
 FPlayerMappableKeyQueryOptions::FPlayerMappableKeyQueryOptions()
 	: MappingName(NAME_None)
 	, KeyToMatch(EKeys::Invalid)
+	, SlotToMatch(EPlayerMappableKeySlot::Unspecified)
 	, bMatchBasicKeyTypes(true)
 	, bMatchKeyAxisType(false)
 	, RequiredDeviceType(EHardwareDevicePrimaryType::Unspecified)
diff --git a/Engine/Plugins/EnhancedInput/Source/EnhancedInput/Public/UserSettings/EnhancedInputUserSettings.h b/Engine/Plugins/EnhancedInput/Source/EnhancedInput/Public/UserSettings/EnhancedInputUserSettings.h
--- a/Engine/Plugins/EnhancedInput/Source/EnhancedInput/Public/UserSettings/EnhancedInputUserSettings.h
+++ b/Engine/Plugins/EnhancedInput/Source/EnhancedInput/Public/UserSettings/EnhancedInputUserSettings.h
@@ -229,6 +229,10 @@ struct ENHANCEDINPUT_API FPlayerMappableKeyQueryOptions
 	UPROPERTY(BlueprintReadWrite, EditAnywhere, Category="Query Options")
 	FKey KeyToMatch;
 	
+	/** The key slot that will be required to match if set. By default this is EPlayerMappableKeySlot::Unspecified, which will not filter by the slot at all. */
+	UPROPERTY(BlueprintReadWrite, EditAnywhere, Category="Query Options")
+	EPlayerMappableKeySlot SlotToMatch;
+
 	/** If true, then only keys that have the same value for IsGamepadKey, IsTouch, and IsGesture will be included in the results of this query */
 	UPROPERTY(BlueprintReadWrite, EditAnywhere, Category="Query Options")
 	uint8 bMatchBasicKeyTypes : 1;
