Exposed saving headers to allow subclasses access.
UpdateOriginalKey now sets the display category of a mapping.
Fixed Typos.

diff --git a/Engine/Plugins/EnhancedInput/Source/EnhancedInput/Private/UserSettings/EnhancedInputUserSettings.cpp b/Engine/Plugins/EnhancedInput/Source/EnhancedInput/Private/UserSettings/EnhancedInputUserSettings.cpp
--- a/Engine/Plugins/EnhancedInput/Source/EnhancedInput/Private/UserSettings/EnhancedInputUserSettings.cpp
+++ b/Engine/Plugins/EnhancedInput/Source/EnhancedInput/Private/UserSettings/EnhancedInputUserSettings.cpp
@@ -31,12 +31,12 @@ namespace UE::EnhancedInput
 	/** The name of the slot that these settings will save to */
 	static const FString SETTINGS_SLOT_NAME = TEXT("EnhancedInputUserSettings");
 
-	UE_DEFINE_GAMEPLAY_TAG_STATIC(TAG_DefaultProfileIdentifier, "InputUserSettings.Profiles.Default");
+	UE_DEFINE_GAMEPLAY_TAG(TAG_DefaultProfileIdentifier, "InputUserSettings.Profiles.Default");
 	static const FText DefaultProfileDisplayName = LOCTEXT("Default_Profile_name", "Default Profile");
 	
-	UE_DEFINE_GAMEPLAY_TAG_STATIC(TAG_InvalidMappingName, "InputUserSettings.FailureReasons.InvalidMappingName");
-	UE_DEFINE_GAMEPLAY_TAG_STATIC(TAG_NoKeyProfile, "InputUserSettings.FailureReasons.NoKeyProfile");
-	UE_DEFINE_GAMEPLAY_TAG_STATIC(TAG_NoMatchingMappings, "InputUserSettings.FailureReasons.NoMatchingMappings");
+	UE_DEFINE_GAMEPLAY_TAG(TAG_InvalidMappingName, "InputUserSettings.FailureReasons.InvalidMappingName");
+	UE_DEFINE_GAMEPLAY_TAG(TAG_NoKeyProfile, "InputUserSettings.FailureReasons.NoKeyProfile");
+	UE_DEFINE_GAMEPLAY_TAG(TAG_NoMatchingMappings, "InputUserSettings.FailureReasons.NoMatchingMappings");
 	
 	// TODO: Break this out to somewhere else, probably the EI library. 
 	static ULocalPlayer* GetLocalPlayer(const UEnhancedPlayerInput* PlayerInput)
@@ -213,6 +213,7 @@ void FPlayerKeyMapping::UpdateOriginalKey(const FEnhancedActionKeyMapping& Origi
 	
 	DefaultKey = OriginalMapping.Key;
 	DisplayName = OriginalMapping.GetDisplayName();
+	DisplayCategory = OriginalMapping.GetDisplayCategory();
 	AssociatedInputAction = OriginalMapping.Action;
 
 	// If the mapping is the same as the default key, make sure that it is not marked as
@@ -636,46 +637,6 @@ void UEnhancedInputUserSettings::OnAsyncSaveComplete(const FString& SlotName, co
 	}
 }
 
-namespace UE::EnhancedInput
-{
-	static const int32 GPlayerMappableSaveVersion = 1;
-
-	struct FKeyMappingSaveData
-	{
-		friend FArchive& operator<<(FArchive& Ar, FKeyMappingSaveData& Data)
-		{
-			Ar << Data.ActionName;
-			Ar << Data.CurrentKeyName;
-			Ar << Data.HardwareDeviceId;
-			Ar << Data.Slot;
-			return Ar;
-		}
-		
-		FName ActionName = NAME_None;
-		FName CurrentKeyName = NAME_None;
-		FHardwareDeviceIdentifier HardwareDeviceId = FHardwareDeviceIdentifier::Invalid;
-		EPlayerMappableKeySlot Slot = EPlayerMappableKeySlot::Unspecified;
-	};
-
-	/** Struct used to store info about the mappable profile subobjects */
-	struct FMappableKeysHeader
-	{
-		friend FArchive& operator<<(FArchive& Ar, FMappableKeysHeader& Header)
-		{
-			Ar << Header.ProfileIdentifier;
-			Ar << Header.ClassPath;
-			Ar << Header.ObjectPath;
-			Ar << Header.DirtyMappings;
-			return Ar;
-		}
-
-		FGameplayTag ProfileIdentifier;
-		FString ClassPath;
-		FString ObjectPath;
-		TArray<FKeyMappingSaveData> DirtyMappings;
-	};
-}
-
 void UEnhancedInputUserSettings::Serialize(FArchive& Ar)
 {
 	Super::Serialize(Ar);
diff --git a/Engine/Plugins/EnhancedInput/Source/EnhancedInput/Public/UserSettings/EnhancedInputUserSettings.h b/Engine/Plugins/EnhancedInput/Source/EnhancedInput/Public/UserSettings/EnhancedInputUserSettings.h
--- a/Engine/Plugins/EnhancedInput/Source/EnhancedInput/Public/UserSettings/EnhancedInputUserSettings.h
+++ b/Engine/Plugins/EnhancedInput/Source/EnhancedInput/Public/UserSettings/EnhancedInputUserSettings.h
@@ -7,6 +7,7 @@
 #include "InputCoreTypes.h"					// For FKey
 #include "GameFramework/InputSettings.h"	// For FHardwareDeviceIdentifier
 #include "GameplayTagContainer.h"
+#include "NativeGameplayTags.h"				// For Extern Tag Declaration Macro
 
 #include "EnhancedInputUserSettings.generated.h"
 
@@ -42,6 +43,51 @@ enum class EPlayerMappableKeySlot : uint8
 	Max
 };
 
+namespace UE::EnhancedInput
+{
+	ENHANCEDINPUT_API UE_DECLARE_GAMEPLAY_TAG_EXTERN(TAG_DefaultProfileIdentifier);
+	ENHANCEDINPUT_API UE_DECLARE_GAMEPLAY_TAG_EXTERN(TAG_InvalidMappingName);
+	ENHANCEDINPUT_API UE_DECLARE_GAMEPLAY_TAG_EXTERN(TAG_NoKeyProfile);
+	ENHANCEDINPUT_API UE_DECLARE_GAMEPLAY_TAG_EXTERN(TAG_NoMatchingMappings);
+
+	static const int32 GPlayerMappableSaveVersion = 1;
+
+	struct FKeyMappingSaveData
+	{
+		friend FArchive& operator<<(FArchive& Ar, FKeyMappingSaveData& Data)
+		{
+			Ar << Data.ActionName;
+			Ar << Data.CurrentKeyName;
+			Ar << Data.HardwareDeviceId;
+			Ar << Data.Slot;
+			return Ar;
+		}
+
+		FName ActionName = NAME_None;
+		FName CurrentKeyName = NAME_None;
+		FHardwareDeviceIdentifier HardwareDeviceId = FHardwareDeviceIdentifier::Invalid;
+		EPlayerMappableKeySlot Slot = EPlayerMappableKeySlot::Unspecified;
+	};
+
+	/** Struct used to store info about the mappable profile subobjects */
+	struct FMappableKeysHeader
+	{
+		friend FArchive& operator<<(FArchive& Ar, FMappableKeysHeader& Header)
+		{
+			Ar << Header.ProfileIdentifier;
+			Ar << Header.ClassPath;
+			Ar << Header.ObjectPath;
+			Ar << Header.DirtyMappings;
+			return Ar;
+		}
+
+		FGameplayTag ProfileIdentifier;
+		FString ClassPath;
+		FString ObjectPath;
+		TArray<FKeyMappingSaveData> DirtyMappings;
+	};
+}
+
 /** Arguments that can be used when mapping a player key */
 USTRUCT(BlueprintType)
 struct ENHANCEDINPUT_API FMapPlayerKeyArgs
@@ -114,7 +160,7 @@ public:
 
 	/**
 	 * The unique FName associated with this mapping. This is defined by this mappings owning Input Action
-	 * or the individual Enhanced Action Key Mapping if it is overriden
+	 * or the individual Enhanced Action Key Mapping if it is overridden
 	 */
 	const FName GetMappingName() const;
 
@@ -130,7 +176,7 @@ public:
 	/** Returns the optional hardware device ID that this mapping is specific to */
 	const FHardwareDeviceIdentifier& GetHardwareDeviceId() const; 
 
-	/** Returns the input actoin asset associated with this player key mapping */
+	/** Returns the input action asset associated with this player key mapping */
 	const UInputAction* GetAssociatedInputAction() const;
 
 	/** Resets the current mapping to the default one */
@@ -316,7 +362,7 @@ public:
 
 	/**
 	 * Called during IEnhancedInputSubsystemInterface::RebuildControlMappings, this provides your user settings
-	 * the oppurtunity to decide what key mappings get added for a given FEnhancedActionKeyMapping, and applied to the player.
+	 * the opportunity to decide what key mappings get added for a given FEnhancedActionKeyMapping, and applied to the player.
 	 *
 	 * Override this to change what query options are being used during the controls being rebuilt of enhanced input.
 	 *
@@ -429,8 +475,8 @@ struct ENHANCEDINPUT_API FPlayerMappableKeyProfileCreationArgs
 
 /**
  * The Enhanced Input User Settings class is a place where you can put all of your Input Related settings
- * that you want your user to be able to change. Things like their key mappings, aim sensativity, accessibility
- * settings, etc. This also provies a Registration point for Input Mappings Contexts (IMC) from possibly unloaded
+ * that you want your user to be able to change. Things like their key mappings, aim sensitivity, accessibility
+ * settings, etc. This also provides a Registration point for Input Mappings Contexts (IMC) from possibly unloaded
  * plugins (i.e. Game Feature Plugins). You can register your IMC from a Game Feature Action plugin here, and then
  * have access to all the key mappings available. This is very useful for building settings screens because you can
  * now access all the mappings in your game, even if the entire plugin isn't loaded yet. 
@@ -592,7 +638,7 @@ public:
 	UFUNCTION(BlueprintCallable, Category="Enhanced Input|User Settings")
 	virtual bool UnregisterInputMappingContext(const UInputMappingContext* IMC);
 
-	/** Removes miltiple mapping contexts from the registered mapping contexts */
+	/** Removes multiple mapping contexts from the registered mapping contexts */
 	UFUNCTION(BlueprintCallable, Category="Enhanced Input|User Settings")
 	bool UnregisterInputMappingContexts(const TSet<UInputMappingContext*>& MappingContexts);
 
