Make the EI user settings use just the local player instead of the input component so that it can be initalized earlier in the EI subsystem lifetime

#rnx

diff --git a/Engine/Plugins/EnhancedInput/Source/EnhancedInput/Private/EnhancedInputSubsystems.cpp b/Engine/Plugins/EnhancedInput/Source/EnhancedInput/Private/EnhancedInputSubsystems.cpp
--- a/Engine/Plugins/EnhancedInput/Source/EnhancedInput/Private/EnhancedInputSubsystems.cpp
+++ b/Engine/Plugins/EnhancedInput/Source/EnhancedInput/Private/EnhancedInputSubsystems.cpp
@@ -23,12 +23,20 @@ DEFINE_LOG_CATEGORY(LogWorldSubsystemInput);
 // *
 // **************************************************************************************************
 
+void UEnhancedInputLocalPlayerSubsystem::Initialize(FSubsystemCollectionBase& Collection)
+{
+	Super::Initialize(Collection);
+	
+	if (GetDefault<UEnhancedInputDeveloperSettings>()->bEnableUserSettings)
+	{
+		InitalizeUserSettings();
+	}
+}
+
 void UEnhancedInputLocalPlayerSubsystem::PlayerControllerChanged(APlayerController* NewPlayerController)
 {
 	Super::PlayerControllerChanged(NewPlayerController);
-
-	// Now that there is a valid player controller, we should have access to a valid PlayerInput pointer.
-	// That means that we can listen for changes to that player's settings!
+	
 	if (GetDefault<UEnhancedInputDeveloperSettings>()->bEnableUserSettings)
 	{
 		InitalizeUserSettings();
@@ -55,11 +63,19 @@ UEnhancedInputUserSettings* UEnhancedInputLocalPlayerSubsystem::GetUserSettings(
 
 void UEnhancedInputLocalPlayerSubsystem::InitalizeUserSettings()
 {
+	ULocalPlayer* LP = GetLocalPlayer();
+
+	if (!LP)
+	{
+		UE_LOG(LogEnhancedInput, Error, TEXT("Failed to initalize user settings, there is no valid local player!"));
+		return;
+	}
+	
 	// If the user settings are already created then we just want to update it's pointer to the 
-	// player input object, as that may have changed with a new player controller.
+	// local player object, as that may have changed with a new player controller.
 	if (UserSettings)
 	{
-		UserSettings->Initialize(GetPlayerInput());
+		UserSettings->Initialize(LP);
 		return;
 	}
 
@@ -68,15 +84,8 @@ void UEnhancedInputLocalPlayerSubsystem::InitalizeUserSettings()
 		UE_LOG(LogEnhancedInput, Verbose, TEXT("bEnableUserSettings is set to false, skipping creation of UEnhancedInputUserSettings!"));
 		return;
 	}
-
-	UEnhancedPlayerInput* PlayerInput = GetPlayerInput();
-	if (!PlayerInput)
-	{
-		UE_LOG(LogEnhancedInput, Error, TEXT("Unable to find a valid player input to initalize user settings with!"));
-		return;
-	}
 	
-	UserSettings = UEnhancedInputUserSettings::LoadOrCreateSettings(PlayerInput);
+	UserSettings = UEnhancedInputUserSettings::LoadOrCreateSettings(LP);
 
 	// Bind delegates to the user settings if it was successfully created
 	if (ensure(UserSettings))
diff --git a/Engine/Plugins/EnhancedInput/Source/EnhancedInput/Private/UserSettings/EnhancedInputUserSettings.cpp b/Engine/Plugins/EnhancedInput/Source/EnhancedInput/Private/UserSettings/EnhancedInputUserSettings.cpp
--- a/Engine/Plugins/EnhancedInput/Source/EnhancedInput/Private/UserSettings/EnhancedInputUserSettings.cpp
+++ b/Engine/Plugins/EnhancedInput/Source/EnhancedInput/Private/UserSettings/EnhancedInputUserSettings.cpp
@@ -37,19 +37,6 @@ namespace UE::EnhancedInput
 	UE_DEFINE_GAMEPLAY_TAG_STATIC(TAG_InvalidMappingName, "InputUserSettings.FailureReasons.InvalidMappingName");
 	UE_DEFINE_GAMEPLAY_TAG_STATIC(TAG_NoKeyProfile, "InputUserSettings.FailureReasons.NoKeyProfile");
 	UE_DEFINE_GAMEPLAY_TAG_STATIC(TAG_NoMatchingMappings, "InputUserSettings.FailureReasons.NoMatchingMappings");
-	
-	// TODO: Break this out to somewhere else, probably the EI library. 
-	static ULocalPlayer* GetLocalPlayer(const UEnhancedPlayerInput* PlayerInput)
-	{
-		if (PlayerInput)
-		{
-			if (APlayerController* PC = PlayerInput->GetOuterAPlayerController())
-			{
-				return PC->GetLocalPlayer();
-			}
-		}
-		return nullptr;
-	}
 
 	static void DumpAllKeyProfilesToLog(const TArray<FString>& Args)
 	{
@@ -528,12 +515,10 @@ void UEnhancedPlayerMappableKeyProfile::K2_FindKeyMapping(FPlayerKeyMapping& Out
 ///////////////////////////////////////////////////////////
 // UEnhancedInputUserSettings
 
-UEnhancedInputUserSettings* UEnhancedInputUserSettings::LoadOrCreateSettings(UEnhancedPlayerInput* PlayerInput)
+UEnhancedInputUserSettings* UEnhancedInputUserSettings::LoadOrCreateSettings(ULocalPlayer* LocalPlayer)
 {
 	UEnhancedInputUserSettings* Settings = nullptr;
 
-	ULocalPlayer* LocalPlayer = UE::EnhancedInput::GetLocalPlayer(PlayerInput);
-
 	if (!LocalPlayer)
 	{
 		UE_LOG(LogEnhancedInput, Log, TEXT("Unable to determine an owning Local Player for the given Enhanced Player Input object"));
@@ -559,17 +544,24 @@ UEnhancedInputUserSettings* UEnhancedInputUserSettings::LoadOrCreateSettings(UEn
 
 	if (ensure(Settings))
 	{
-		Settings->Initialize(PlayerInput);
+		Settings->Initialize(LocalPlayer);
     	Settings->ApplySettings();	
 	}
 
 	return Settings;
 }
 
-void UEnhancedInputUserSettings::Initialize(UEnhancedPlayerInput* InPlayerInput)
+void UEnhancedInputUserSettings::Initialize(ULocalPlayer* InLocalPlayer)
 {
-	OwningPlayerInput = InPlayerInput;
-	ensureMsgf(GetPlayerInput(), TEXT("UEnhancedInputUserSettings is missing a player input!"));
+	// If the local player hasn't changed, then we don't need to do anything.
+	// This may be the case if the Player Controller has changed
+	if (InLocalPlayer == GetLocalPlayer())
+	{
+		return;
+	}
+	
+	OwningLocalPlayer = InLocalPlayer;
+	ensureMsgf(OwningLocalPlayer.IsValid(), TEXT("UEnhancedInputUserSettings is missing an owning local player!"));
 
 	// Create a default key mapping profile in the case where one doesn't exist
 	if (!GetCurrentKeyProfile())
@@ -588,7 +580,6 @@ void UEnhancedInputUserSettings::Initialize(UEnhancedPlayerInput* InPlayerInput)
 
 void UEnhancedInputUserSettings::ApplySettings()
 {
-	ensureMsgf(GetPlayerInput(), TEXT("UEnhancedInputUserSettings is missing a player input!"));
 	UE_LOG(LogEnhancedInput, Verbose, TEXT("Enhanced Input User Settings applied!"));
 
 	OnSettingsApplied.Broadcast();
@@ -596,8 +587,6 @@ void UEnhancedInputUserSettings::ApplySettings()
 
 void UEnhancedInputUserSettings::SaveSettings()
 {
-	ensureMsgf(GetPlayerInput(), TEXT("UEnhancedInputUserSettings is missing a player input!"));
-	
 	if (ULocalPlayer* OwningPlayer = GetLocalPlayer())
 	{
 		UGameplayStatics::SaveGameToSlot(this, UE::EnhancedInput::SETTINGS_SLOT_NAME, OwningPlayer->GetLocalPlayerIndex());
@@ -773,27 +762,9 @@ void UEnhancedInputUserSettings::Serialize(FArchive& Ar)
 	}
 }
 
-UEnhancedPlayerInput* UEnhancedInputUserSettings::GetPlayerInput() const
-{
-	return OwningPlayerInput.Get();
-}
-
-APlayerController* UEnhancedInputUserSettings::GetPlayerController() const
-{
-	if (UEnhancedPlayerInput* PlayerInput = GetPlayerInput())
-	{
-		return Cast<APlayerController>(PlayerInput->GetOuter());
-	}
-	return nullptr;
-}
-
 ULocalPlayer* UEnhancedInputUserSettings::GetLocalPlayer() const
 {
-	if (APlayerController* PC = GetPlayerController())
-	{
-		return PC->GetLocalPlayer();
-	}
-	return nullptr;
+	return OwningLocalPlayer.Get();
 }
 
 void UEnhancedInputUserSettings::MapPlayerKey(const FMapPlayerKeyArgs& InArgs, FGameplayTagContainer& FailureReason)
@@ -941,10 +912,9 @@ const UInputAction* UEnhancedInputUserSettings::FindInputActionForMapping(const
 
 bool UEnhancedInputUserSettings::SetKeyProfile(const FGameplayTag& InProfileId)
 {
-	UEnhancedPlayerInput* PlayerInput = GetPlayerInput();
-	if (!PlayerInput)
+	if (!GetLocalPlayer())
 	{
-		UE_LOG(LogEnhancedInput, Error, TEXT("Failed to find the PlayerInput associated with the Enhanced Input user settings!"), *InProfileId.ToString());
+		UE_LOG(LogEnhancedInput, Error, TEXT("Failed to find the Local Player associated with the Enhanced Input user settings!"), *InProfileId.ToString());
 		return false;
 	}
 
diff --git a/Engine/Plugins/EnhancedInput/Source/EnhancedInput/Public/EnhancedInputSubsystems.h b/Engine/Plugins/EnhancedInput/Source/EnhancedInput/Public/EnhancedInputSubsystems.h
--- a/Engine/Plugins/EnhancedInput/Source/EnhancedInput/Public/EnhancedInputSubsystems.h
+++ b/Engine/Plugins/EnhancedInput/Source/EnhancedInput/Public/EnhancedInputSubsystems.h
@@ -23,6 +23,7 @@ class ENHANCEDINPUT_API UEnhancedInputLocalPlayerSubsystem : public ULocalPlayer
 public:
 
 	// Begin ULocalPlayerSubsystem
+	virtual void Initialize(FSubsystemCollectionBase& Collection) override;
 	virtual void PlayerControllerChanged(APlayerController* NewPlayerController) override;
 	// End ULocalPlayerSubsystem
 	
diff --git a/Engine/Plugins/EnhancedInput/Source/EnhancedInput/Public/UserSettings/EnhancedInputUserSettings.h b/Engine/Plugins/EnhancedInput/Source/EnhancedInput/Public/UserSettings/EnhancedInputUserSettings.h
--- a/Engine/Plugins/EnhancedInput/Source/EnhancedInput/Public/UserSettings/EnhancedInputUserSettings.h
+++ b/Engine/Plugins/EnhancedInput/Source/EnhancedInput/Public/UserSettings/EnhancedInputUserSettings.h
@@ -453,8 +453,8 @@ public:
 	//~ End UObject Interface
 	
 	/** Loads or creates new user settings for the owning local player of the given player input */
-	static UEnhancedInputUserSettings* LoadOrCreateSettings(UEnhancedPlayerInput* PlayerInput);
-	virtual void Initialize(UEnhancedPlayerInput* PlayerInput);
+	static UEnhancedInputUserSettings* LoadOrCreateSettings(ULocalPlayer* LP);
+	virtual void Initialize(ULocalPlayer* LP);
 
 	/**
 	 * Apply any custom input settings to your user. By default, this will just broadcast the OnSettingsApplied delegate
@@ -485,9 +485,7 @@ protected:
 
 public:
 	
-	UEnhancedPlayerInput* GetPlayerInput() const;
 	ULocalPlayer* GetLocalPlayer() const;
-	APlayerController* GetPlayerController() const;
 	
 	/** Fired when the user settings have changed, such as their key mappings. */
 	DECLARE_DYNAMIC_MULTICAST_DELEGATE_OneParam(FEnhancedInputUserSettingsChanged, UEnhancedInputUserSettings*, Settings);
@@ -615,9 +613,9 @@ protected:
 	UPROPERTY(SaveGame)
 	TMap<FGameplayTag, TObjectPtr<UEnhancedPlayerMappableKeyProfile>> SavedKeyProfiles;
 	
-	/** The owning Player Input object of these settings */
+	/** The owning Local Player object of these settings */
 	UPROPERTY(Transient)
-	TWeakObjectPtr<UEnhancedPlayerInput> OwningPlayerInput;
+	TWeakObjectPtr<ULocalPlayer> OwningLocalPlayer;
 	
 	/**
 	 * Set of currently registered input mapping contexts that may not be currently
