- Add some setters for Hardware device ID's on player mappable keys
- Initalize bIsDirty to false in the constructor
- Change IsCustomized to allow for the EKeys::Invalid key, which means that players can set a keybinding to "Unbound" if they desire. This means that the "CustomKey" will always be set to the default key now, and we check if it is different then the default key and not check it for validity
- Fix a few simple typos and make some structs final to make it impossible to subclass them, since it could risk serialization issues.

diff --git a/Engine/Plugins/EnhancedInput/Source/EnhancedInput/Private/UserSettings/EnhancedInputUserSettings.cpp b/Engine/Plugins/EnhancedInput/Source/EnhancedInput/Private/UserSettings/EnhancedInputUserSettings.cpp
--- a/Engine/Plugins/EnhancedInput/Source/EnhancedInput/Private/UserSettings/EnhancedInputUserSettings.cpp
+++ b/Engine/Plugins/EnhancedInput/Source/EnhancedInput/Private/UserSettings/EnhancedInputUserSettings.cpp
@@ -79,6 +79,7 @@ FPlayerKeyMapping::FPlayerKeyMapping()
 	, DisplayName(FText::GetEmpty())
 	, DisplayCategory(FText::GetEmpty())
 	, Slot(EPlayerMappableKeySlot::Unspecified)
+	, bIsDirty(false)
 	, DefaultKey(EKeys::Invalid)
 	, CurrentKey(EKeys::Invalid)
 	, HardwareDeviceId(FHardwareDeviceIdentifier::Invalid)
@@ -86,14 +87,18 @@ FPlayerKeyMapping::FPlayerKeyMapping()
 {
 }
 
-FPlayerKeyMapping::FPlayerKeyMapping(const FEnhancedActionKeyMapping& OriginalMapping, EPlayerMappableKeySlot InSlot /* = EPlayerMappableKeySlot::Unspecified */)
+FPlayerKeyMapping::FPlayerKeyMapping(
+	const FEnhancedActionKeyMapping& OriginalMapping,
+	EPlayerMappableKeySlot InSlot /* = EPlayerMappableKeySlot::Unspecified */,
+	const FHardwareDeviceIdentifier& InHardwareDevice /* = FHardwareDeviceIdentifier::Invalid */)
 	: MappingName(OriginalMapping.GetMappingName())
 	, DisplayName(OriginalMapping.GetDisplayName())
 	, DisplayCategory(OriginalMapping.GetDisplayCategory())
 	, Slot(InSlot)
+	, bIsDirty(false)
 	, DefaultKey(OriginalMapping.Key)
-	, CurrentKey(EKeys::Invalid)
-	, HardwareDeviceId(FHardwareDeviceIdentifier::Invalid)
+	, CurrentKey(OriginalMapping.Key)
+	, HardwareDeviceId(InHardwareDevice)
 	, AssociatedInputAction(OriginalMapping.Action)
 {
 }
@@ -104,7 +109,7 @@ FPlayerKeyMapping FPlayerKeyMapping::InvalidMapping = FPlayerKeyMapping();
 
 bool FPlayerKeyMapping::IsCustomized() const
 {
-	return CurrentKey.IsValid() && (CurrentKey != DefaultKey);
+	return (CurrentKey != DefaultKey);
 }
 
 bool FPlayerKeyMapping::IsValid() const
@@ -161,6 +166,16 @@ const FHardwareDeviceIdentifier& FPlayerKeyMapping::GetHardwareDeviceId() const
 	return HardwareDeviceId;
 }
 
+EHardwareDevicePrimaryType FPlayerKeyMapping::GetPrimaryDeviceType() const
+{
+	return HardwareDeviceId.PrimaryDeviceType;
+}
+
+EHardwareDeviceSupportedFeatures::Type FPlayerKeyMapping::GetHardwareDeviceSupportedFeatures() const
+{
+	return static_cast<EHardwareDeviceSupportedFeatures::Type>(HardwareDeviceId.SupportedFeaturesMask);
+}
+
 const UInputAction* FPlayerKeyMapping::GetAssociatedInputAction() const
 {
 	return AssociatedInputAction;
@@ -190,8 +205,17 @@ void FPlayerKeyMapping::ResetToDefault()
 
 void FPlayerKeyMapping::SetCurrentKey(const FKey& NewKey)
 {
+	if (NewKey != CurrentKey)
+	{
+		bIsDirty = true;
+	}
+	
 	CurrentKey = NewKey;
-	bIsDirty = true;
+}
+
+void FPlayerKeyMapping::SetHardwareDeviceId(const FHardwareDeviceIdentifier& InDeviceId)
+{
+	HardwareDeviceId = InDeviceId;
 }
 
 void FPlayerKeyMapping::UpdateOriginalKey(const FEnhancedActionKeyMapping& OriginalMapping)
@@ -199,6 +223,7 @@ void FPlayerKeyMapping::UpdateOriginalKey(const FEnhancedActionKeyMapping& Origi
 	ensure(OriginalMapping.IsPlayerMappable() && OriginalMapping.GetMappingName() == MappingName);
 	
 	DefaultKey = OriginalMapping.Key;
+	DisplayCategory = OriginalMapping.GetDisplayCategory();
 	DisplayName = OriginalMapping.GetDisplayName();
 	AssociatedInputAction = OriginalMapping.Action;
 
@@ -791,7 +816,7 @@ void UEnhancedInputUserSettings::MapPlayerKey(const FMapPlayerKeyArgs& InArgs, F
 		OnSettingsChanged.Broadcast(this);
 	}
 	// If it doesn't exist, then we need to make it if there is a valid action name
-	else if (FKeyMappingRow* MappingRow = KeyProfile->PlayerMappedKeys.Find(InArgs.MappingName))
+	else if (FKeyMappingRow* MappingRow = KeyProfile->FindKeyMappingRowMutable(InArgs.MappingName))
 	{
 		// If one doesn't exist, then we need to create a new mapping in the given slot.
 		// 
diff --git a/Engine/Plugins/EnhancedInput/Source/EnhancedInput/Public/UserSettings/EnhancedInputUserSettings.h b/Engine/Plugins/EnhancedInput/Source/EnhancedInput/Public/UserSettings/EnhancedInputUserSettings.h
--- a/Engine/Plugins/EnhancedInput/Source/EnhancedInput/Public/UserSettings/EnhancedInputUserSettings.h
+++ b/Engine/Plugins/EnhancedInput/Source/EnhancedInput/Public/UserSettings/EnhancedInputUserSettings.h
@@ -44,7 +44,7 @@ enum class EPlayerMappableKeySlot : uint8
 
 /** Arguments that can be used when mapping a player key */
 USTRUCT(BlueprintType)
-struct ENHANCEDINPUT_API FMapPlayerKeyArgs
+struct ENHANCEDINPUT_API FMapPlayerKeyArgs final
 {
 	GENERATED_BODY()
 
@@ -80,7 +80,7 @@ struct ENHANCEDINPUT_API FMapPlayerKeyArgs
 
 /** Represents a single key mapping that is set by the player */
 USTRUCT(BlueprintType)
-struct ENHANCEDINPUT_API FPlayerKeyMapping
+struct ENHANCEDINPUT_API FPlayerKeyMapping final
 {
 	GENERATED_BODY()
 	friend class UEnhancedInputUserSettings;
@@ -89,7 +89,10 @@ public:
 	
 	FPlayerKeyMapping();
 	
-	FPlayerKeyMapping(const FEnhancedActionKeyMapping& OriginalMapping, EPlayerMappableKeySlot InSlot = EPlayerMappableKeySlot::Unspecified);
+	FPlayerKeyMapping(
+		const FEnhancedActionKeyMapping& OriginalMapping,
+		EPlayerMappableKeySlot InSlot = EPlayerMappableKeySlot::Unspecified,
+		const FHardwareDeviceIdentifier& InHardwareDevice = FHardwareDeviceIdentifier::Invalid);
 
 	/** A static invalid player key mapping to be used for easy comparisons in blueprint */
 	static FPlayerKeyMapping InvalidMapping;
@@ -114,7 +117,7 @@ public:
 
 	/**
 	 * The unique FName associated with this mapping. This is defined by this mappings owning Input Action
-	 * or the individual Enhanced Action Key Mapping if it is overriden
+	 * or the individual Enhanced Action Key Mapping if it is overridden
 	 */
 	const FName GetMappingName() const;
 
@@ -128,9 +131,15 @@ public:
 	EPlayerMappableKeySlot GetSlot() const;
 
 	/** Returns the optional hardware device ID that this mapping is specific to */
-	const FHardwareDeviceIdentifier& GetHardwareDeviceId() const; 
+	const FHardwareDeviceIdentifier& GetHardwareDeviceId() const;
 
-	/** Returns the input actoin asset associated with this player key mapping */
+	/** Gets the primary device type associated with this key mapping. Taken from the HardwareDeviceId. */
+	EHardwareDevicePrimaryType GetPrimaryDeviceType() const;
+
+	/** Gets the supported hardware device feature flags associated with this key mapping. Taken from the HardwareDeviceId. */
+	EHardwareDeviceSupportedFeatures::Type GetHardwareDeviceSupportedFeatures() const;
+
+	/** Returns the input action asset associated with this player key mapping */
 	const UInputAction* GetAssociatedInputAction() const;
 
 	/** Resets the current mapping to the default one */
@@ -139,6 +148,9 @@ public:
 	/** Sets the value of the current key to the one given */
 	void SetCurrentKey(const FKey& NewKey);
 
+	/** Sets the value of the hardware device ID that this custom key mapping is associated with. This can be used to filter key mappings per input device */
+	void SetHardwareDeviceId(const FHardwareDeviceIdentifier& InDeviceId);
+
 	/**
 	 * Updates the metadata properties on this player mapped key based on the given
 	 * enhanced action mapping. This will populate the fields on this struct that are not editable
@@ -199,7 +211,7 @@ protected:
  * this needs to be a struct (blueprint don't support nested containers).
  */
 USTRUCT(BlueprintType)
-struct ENHANCEDINPUT_API FKeyMappingRow
+struct ENHANCEDINPUT_API FKeyMappingRow final
 {
 	GENERATED_BODY()
 
@@ -215,7 +227,7 @@ struct ENHANCEDINPUT_API FKeyMappingRow
 * key profile. 
 */
 USTRUCT(BlueprintType)
-struct ENHANCEDINPUT_API FPlayerMappableKeyQueryOptions
+struct ENHANCEDINPUT_API FPlayerMappableKeyQueryOptions final
 {
 	GENERATED_BODY()
 
@@ -316,7 +328,7 @@ public:
 
 	/**
 	 * Called during IEnhancedInputSubsystemInterface::RebuildControlMappings, this provides your user settings
-	 * the oppurtunity to decide what key mappings get added for a given FEnhancedActionKeyMapping, and applied to the player.
+	 * the opportunity to decide what key mappings get added for a given FEnhancedActionKeyMapping, and applied to the player.
 	 *
 	 * Override this to change what query options are being used during the controls being rebuilt of enhanced input.
 	 *
@@ -402,7 +414,7 @@ protected:
 
 /** Arguments that can be used when creating a new mapping profile */
 USTRUCT(BlueprintType)
-struct ENHANCEDINPUT_API FPlayerMappableKeyProfileCreationArgs
+struct ENHANCEDINPUT_API FPlayerMappableKeyProfileCreationArgs final
 {
 	GENERATED_BODY()
 	
@@ -429,8 +441,8 @@ struct ENHANCEDINPUT_API FPlayerMappableKeyProfileCreationArgs
 
 /**
  * The Enhanced Input User Settings class is a place where you can put all of your Input Related settings
- * that you want your user to be able to change. Things like their key mappings, aim sensativity, accessibility
- * settings, etc. This also provies a Registration point for Input Mappings Contexts (IMC) from possibly unloaded
+ * that you want your user to be able to change. Things like their key mappings, aim sensitivity, accessibility
+ * settings, etc. This also provides a Registration point for Input Mappings Contexts (IMC) from possibly unloaded
  * plugins (i.e. Game Feature Plugins). You can register your IMC from a Game Feature Action plugin here, and then
  * have access to all the key mappings available. This is very useful for building settings screens because you can
  * now access all the mappings in your game, even if the entire plugin isn't loaded yet. 
@@ -590,7 +602,7 @@ public:
 	UFUNCTION(BlueprintCallable, Category="Enhanced Input|User Settings")
 	virtual bool UnregisterInputMappingContext(const UInputMappingContext* IMC);
 
-	/** Removes miltiple mapping contexts from the registered mapping contexts */
+	/** Removes multiple mapping contexts from the registered mapping contexts */
 	UFUNCTION(BlueprintCallable, Category="Enhanced Input|User Settings")
 	bool UnregisterInputMappingContexts(const TSet<UInputMappingContext*>& MappingContexts);
 
