Add a debug command to display the user's input settings. This will display all of the currently mapped keys by the player, what their default keys are, and if they are dirty or not.

Add another command to save the settings to a slot to make testing easier, without having to create some UI to do it in game

#rnx

diff --git a/Engine/Plugins/EnhancedInput/Config/Input.ini b/Engine/Plugins/EnhancedInput/Config/Input.ini
--- a/Engine/Plugins/EnhancedInput/Config/Input.ini
+++ b/Engine/Plugins/EnhancedInput/Config/Input.ini
@@ -1,3 +1,4 @@
 [/Script/EngineSettings.ConsoleSettings]
 +ManualAutoCompleteList=(Command="ShowDebug EnhancedInput",Desc="Displays debug information about the current state of any Enhanced Input Mapping Contexts")
 +ManualAutoCompleteList=(Command="ShowDebug WorldSubsystemInput",Desc="Displays debug information about the current state of any Enhanced Input Mapping Contexts applied to the Enhanced Input world subsystem.")
++ManualAutoCompleteList=(Command="ShowDebug InputSettings",Desc="Displays debug information any user input settings, such as player mappable keys.")
diff --git a/Engine/Plugins/EnhancedInput/Source/EnhancedInput/Private/EnhancedInputModule.cpp b/Engine/Plugins/EnhancedInput/Source/EnhancedInput/Private/EnhancedInputModule.cpp
--- a/Engine/Plugins/EnhancedInput/Source/EnhancedInput/Private/EnhancedInputModule.cpp
+++ b/Engine/Plugins/EnhancedInput/Source/EnhancedInput/Private/EnhancedInputModule.cpp
@@ -13,6 +13,7 @@
 #include "UObject/Package.h"
 #include "UObject/UObjectIterator.h"
 #include "EnhancedInputDeveloperSettings.h"
+#include "UserSettings/EnhancedInputUserSettings.h"
 
 #define LOCTEXT_NAMESPACE "EnhancedInput"
 
@@ -321,6 +322,7 @@ void FEnhancedInputModule::OnShowDebugInfo(AHUD* HUD, UCanvas* Canvas, const FDe
 {
 	static const FName NAME_EnhancedInput("EnhancedInput");
 	static const FName NAME_WorldSubsystemInput("WorldSubsystemInput");
+	static const FName NAME_InputSettings("InputSettings");
 
 	if (Canvas)
 	{
@@ -348,6 +350,31 @@ void FEnhancedInputModule::OnShowDebugInfo(AHUD* HUD, UCanvas* Canvas, const FDe
 				It->ShowDebugInfo(Canvas);
 			}
 		}
+
+#if ENABLE_DRAW_DEBUG
+		// Show the debug info for input user settings
+		if (HUD->ShouldDisplayDebug(NAME_InputSettings))
+		{
+			// TODO: some way to page through different local players instead of
+			// only showing the first player
+			TObjectIterator<UEnhancedInputLocalPlayerSubsystem> FirstPlayer;
+			if (FirstPlayer)
+			{
+				if (UEnhancedInputUserSettings* Settings = FirstPlayer->GetUserSettings())
+				{
+					Settings->ShowDebugInfo(Canvas);
+				}
+				else
+				{
+					FDisplayDebugManager& DisplayDebugManager = Canvas->DisplayDebugManager;
+					DisplayDebugManager.SetFont(GEngine->GetLargeFont());
+					DisplayDebugManager.SetDrawColor(FColor::Red);
+					DisplayDebugManager.DrawString(TEXT("No Input User Settings found! Did you enable it in project settings?"));
+				}
+			}
+		}
+#endif	// ENABLE_DRAW_DEBUG
+
 	}
 	else
 	{
diff --git a/Engine/Plugins/EnhancedInput/Source/EnhancedInput/Private/UserSettings/EnhancedInputUserSettings.cpp b/Engine/Plugins/EnhancedInput/Source/EnhancedInput/Private/UserSettings/EnhancedInputUserSettings.cpp
--- a/Engine/Plugins/EnhancedInput/Source/EnhancedInput/Private/UserSettings/EnhancedInputUserSettings.cpp
+++ b/Engine/Plugins/EnhancedInput/Source/EnhancedInput/Private/UserSettings/EnhancedInputUserSettings.cpp
@@ -2,7 +2,6 @@
 
 #include "UserSettings/EnhancedInputUserSettings.h"
 
-#include "EnhancedPlayerInput.h"
 #include "EnhancedInputModule.h"
 #include "InputMappingContext.h"
 #include "InputAction.h"
@@ -10,9 +9,7 @@
 #include "EnhancedActionKeyMapping.h"
 #include "EnhancedInputSubsystems.h"
 #include "Engine/LocalPlayer.h"
-#include "GameFramework/PlayerController.h"
 #include "Kismet/GameplayStatics.h"
-#include "GameFramework/InputDeviceSubsystem.h"
 #include "EnhancedInputDeveloperSettings.h"
 #include "Internationalization/Text.h"
 #include "GameFramework/InputSettings.h"
@@ -20,9 +17,13 @@
 #include "EnhancedInputLibrary.h"
 #include "Algo/Find.h"
 #include "NativeGameplayTags.h"
-#include "PlatformFeatures.h"
 #include "SaveGameSystem.h"
-#include "EnhancedActionKeyMapping.h"
+#include "Engine/Engine.h"		// for GEngine definition
+#include "DrawDebugHelpers.h"	// required for ENABLE_DRAW_DEBUG define
+
+#if ENABLE_DRAW_DEBUG
+#include "Engine/Canvas.h"
+#endif
 
 #define LOCTEXT_NAMESPACE "EnhancedInputMappableUserSettings"
 
@@ -53,10 +54,26 @@ namespace UE::EnhancedInput
 		});
 	}
 	
-	static FAutoConsoleCommand ConsoelCommandDumpProfileToLog(
+	static void SaveAllKeyProfilesToSlot(const TArray<FString>& Args)
+	{
+		UEnhancedInputLibrary::ForEachSubsystem([Args](IEnhancedInputSubsystemInterface* Subsystem)
+		{
+			if (UEnhancedInputUserSettings* Settings = Subsystem->GetUserSettings())
+			{
+				Settings->AsyncSaveSettings();
+			}
+		});
+	}
+	
+	static FAutoConsoleCommand ConsoleCommandDumpProfileToLog(
 		TEXT("EnhancedInput.DumpKeyProfileToLog"),
 		TEXT(""),
 		FConsoleCommandWithArgsDelegate::CreateStatic(UE::EnhancedInput::DumpAllKeyProfilesToLog));
+
+	static FAutoConsoleCommand ConsoleCommandDumpSaveKeyProfilesToSlot(
+		TEXT("EnhancedInput.SaveKeyProfilesToSlot"),
+		TEXT("Save the user input settings object with the Save Game to slot system"),
+		FConsoleCommandWithArgsDelegate::CreateStatic(UE::EnhancedInput::SaveAllKeyProfilesToSlot));
 }
 
 ///////////////////////////////////////////////////////////
@@ -1184,6 +1201,146 @@ bool UEnhancedInputUserSettings::IsMappingContextRegistered(const UInputMappingC
 	return RegisteredMappingContexts.Contains(IMC);
 }
 
+#if ENABLE_DRAW_DEBUG
+
+namespace UE::EnhancedInput::DebugColors
+{
+	static const FColor& HeaderColor = FColor::Magenta;
+	static const FColor& TitleColor = FColor::Orange;
+	static const FColor& NormalColor = FColor::Silver;
+	static const FColor& ErrorColor = FColor::Red;
+	static const FColor& KeyMappingRow = FColor::Emerald;
+	static const FColor& KeyMappingColor = FColor::Silver;
+};
+
+#endif	// ENABLE_DRAW_DEBUG
+
+void UEnhancedInputUserSettings::ShowDebugInfo(UCanvas* Canvas) const
+{
+#if ENABLE_DRAW_DEBUG
+
+	if (!Canvas)
+	{
+		return;
+	}
+
+	FDisplayDebugManager& DisplayDebugManager = Canvas->DisplayDebugManager;
+
+	const ULocalPlayer* LP = GetLocalPlayer();	
+
+	// Draw the "Header" of this debug info
+	{
+		DisplayDebugManager.SetFont(GEngine->GetLargeFont());
+		DisplayDebugManager.SetDrawColor(UE::EnhancedInput::DebugColors::HeaderColor);
+		DisplayDebugManager.DrawString(FString::Printf(TEXT("\n\nInput User Settings (%s)\n------------------------------"), *GetNameSafe(GetClass())));
+	}
+
+	// Draw some info about the local player, return if it doesn't exist.
+	{
+		DisplayDebugManager.SetDrawColor(UE::EnhancedInput::DebugColors::TitleColor);
+		DisplayDebugManager.DrawString(TEXT("\t Local Player:"));
+
+		if (!LP)
+		{
+			DisplayDebugManager.SetDrawColor(UE::EnhancedInput::DebugColors::ErrorColor);
+			DisplayDebugManager.DrawString(TEXT("Invalid Local Player on the Input User Settings!"));
+			return;
+		}
+		else
+		{			
+			DisplayDebugManager.SetDrawColor(UE::EnhancedInput::DebugColors::NormalColor);
+			DisplayDebugManager.DrawString(FString::Printf(TEXT("\t \t %s \n \t PlatformUserId: %d"), *LP->GetFName().ToString(), LP->GetPlatformUserId().GetInternalId()));			
+		}
+	}
+
+	ShowDebugInfoInternal(Canvas);
+
+	// Draw the "Footer" of this debug info
+	{
+		DisplayDebugManager.SetDrawColor(UE::EnhancedInput::DebugColors::HeaderColor);
+		DisplayDebugManager.DrawString(TEXT("----------------------------\n(end of Input User Settings)\n------------------------------\n"));
+	}	
+
+#endif	// ENABLE_DRAW_DEBUG
+}
+
+void UEnhancedInputUserSettings::ShowDebugInfoInternal(UCanvas* Canvas) const
+{
+#if ENABLE_DRAW_DEBUG
+	check(Canvas);
+
+	FDisplayDebugManager& DisplayDebugManager = Canvas->DisplayDebugManager;
+
+	DisplayDebugManager.SetDrawColor(UE::EnhancedInput::DebugColors::TitleColor);
+	DisplayDebugManager.DrawString(FString::Printf(TEXT("Player Key Profiles")));
+	DisplayDebugManager.SetDrawColor(UE::EnhancedInput::DebugColors::NormalColor);
+
+	// number of key profiles	
+	DisplayDebugManager.DrawString(FString::Printf(TEXT("\t There are '%d' saved key profiles"), SavedKeyProfiles.Num()));
+
+	// currently active profile
+	if (UEnhancedPlayerMappableKeyProfile* CurrentProfile = GetCurrentKeyProfile())
+	{
+		DisplayDebugManager.DrawString(FString::Printf(TEXT("\t Currently Active Profile: %s (%s)"), *CurrentProfile->GetProfileIdentifer().ToString(), *GetNameSafe(CurrentProfile->GetClass())));
+	}
+	else
+	{
+		DisplayDebugManager.SetDrawColor(UE::EnhancedInput::DebugColors::ErrorColor);
+		DisplayDebugManager.DrawString(TEXT("There is no currently active profile!"));
+	}
+
+	// Draw debug info about the player mappings
+	for (const TPair<FGameplayTag, TObjectPtr<UEnhancedPlayerMappableKeyProfile>>& ProfilePair : SavedKeyProfiles)
+	{
+		UEnhancedPlayerMappableKeyProfile* Profile = ProfilePair.Value;
+		if (!Profile)
+		{
+			continue;
+		}
+
+		const UEnum* const SlotEnum = FindObjectChecked<UEnum>(nullptr, TEXT("/Script/EnhancedInput.EPlayerMappableKeySlot"));
+		const auto& GetSlotString = [&SlotEnum](const EPlayerMappableKeySlot Value) -> FString
+		{
+			if (!SlotEnum)
+			{
+				return TEXT("Unknown Slot Enum State!");
+			}
+
+			return SlotEnum->GetNameStringByValue(static_cast<int64>(Value));
+		};
+
+		DisplayDebugManager.SetDrawColor(UE::EnhancedInput::DebugColors::TitleColor);
+		DisplayDebugManager.DrawString(FString::Printf(TEXT("Profile: %s (%s)\n"), *Profile->GetProfileIdentifer().ToString(), *GetNameSafe(Profile->GetClass())));
+		
+		DisplayDebugManager.SetDrawColor(UE::EnhancedInput::DebugColors::NormalColor);
+		DisplayDebugManager.DrawString(TEXT("\t -------------------- "));
+
+		DisplayDebugManager.SetDrawColor(UE::EnhancedInput::DebugColors::TitleColor);
+		DisplayDebugManager.DrawString(TEXT("\t Player Key Mappings:\n"));
+
+		for (const TPair<FName, FKeyMappingRow>& Pair : Profile->PlayerMappedKeys)
+		{			
+			DisplayDebugManager.SetDrawColor(UE::EnhancedInput::DebugColors::KeyMappingRow);
+			DisplayDebugManager.DrawString(FString::Printf(TEXT("Mapping Row: '%s' \t \t has '%d' Mappings"), *Pair.Key.ToString(), Pair.Value.Mappings.Num()));
+
+			for (const FPlayerKeyMapping& Mapping : Pair.Value.Mappings)
+			{
+				TStringBuilder<1024> Builder;		
+				Builder.Appendf(TEXT("\t| Slot: %s   \t"), *GetSlotString(Mapping.GetSlot()));
+				Builder.Appendf(TEXT("\t| Current: %s   \t"), *Mapping.GetCurrentKey().GetDisplayName().ToString());
+				Builder.Appendf(TEXT("\t| Default: %s   \t"), *Mapping.GetDefaultKey().GetDisplayName().ToString());
+				Builder.Appendf(TEXT("\t| Dirty: %s   "), Mapping.IsDirty() ? TEXT("true") : TEXT("false"));
+			
+				// Flip flop the color for each key mapping to make the debug screen easier to read
+				DisplayDebugManager.SetDrawColor(UE::EnhancedInput::DebugColors::KeyMappingColor);
+				DisplayDebugManager.DrawString(Builder.ToString());
+			}
+		}
+	}
+
+#endif	// ENABLE_DRAW_DEBUG
+}
+
 FPlayerMappableKeyQueryOptions::FPlayerMappableKeyQueryOptions()
 	: MappingName(NAME_None)
 	, KeyToMatch(EKeys::Invalid)
diff --git a/Engine/Plugins/EnhancedInput/Source/EnhancedInput/Public/UserSettings/EnhancedInputUserSettings.h b/Engine/Plugins/EnhancedInput/Source/EnhancedInput/Public/UserSettings/EnhancedInputUserSettings.h
--- a/Engine/Plugins/EnhancedInput/Source/EnhancedInput/Public/UserSettings/EnhancedInputUserSettings.h
+++ b/Engine/Plugins/EnhancedInput/Source/EnhancedInput/Public/UserSettings/EnhancedInputUserSettings.h
@@ -15,6 +15,7 @@ class UEnhancedPlayerInput;
 class ULocalPlayer;
 struct FEnhancedActionKeyMapping;
 class UInputAction;
+class UCanvas;
 
 /**
  * The "Slot" that a player mappable key is in.
@@ -613,8 +614,17 @@ public:
 	UFUNCTION(BlueprintCallable, BlueprintPure, Category="Enhanced Input|User Settings")
 	bool IsMappingContextRegistered(const UInputMappingContext* IMC) const;
 	
+	/** Draw debug information on the screen for these user settings. */
+	void ShowDebugInfo(UCanvas* Canvas) const;
+
 protected:
 
+	/** 
+	* Provides a space for subclasses to add additional debug info if desired.
+	* By default this will draw info about all player mapped keys.
+	*/
+	virtual void ShowDebugInfoInternal(UCanvas* Canvas) const;
+
 	/** The current key profile that is equipped by the user. */
 	UPROPERTY(SaveGame)
 	FGameplayTag CurrentProfileIdentifier;
