Add a "Experimental" display name to the user settings where you can enable/disable it.

#rnx

diff --git a/Engine/Plugins/EnhancedInput/Source/EnhancedInput/Public/EnhancedInputDeveloperSettings.h b/Engine/Plugins/EnhancedInput/Source/EnhancedInput/Public/EnhancedInputDeveloperSettings.h
--- a/Engine/Plugins/EnhancedInput/Source/EnhancedInput/Public/EnhancedInputDeveloperSettings.h
+++ b/Engine/Plugins/EnhancedInput/Source/EnhancedInput/Public/EnhancedInputDeveloperSettings.h
@@ -81,11 +81,11 @@ public:
 	TSoftClassPtr<UEnhancedPlayerInput> DefaultWorldInputClass;
 
 	/**
-	 * If true, then an instance of the User Settings Class will be created on each Enhanced Input subsytem.
+	 * If true, then an instance of the User Settings Class will be created on each Enhanced Input subsystem.
 	 * 
 	 * Note: This is a new experimental feature!
 	 */
-	UPROPERTY(config, EditAnywhere, Category = "Enhanced Input|User Settings", meta=(DisplayPriority = 1))
+	UPROPERTY(config, EditAnywhere, Category = "Enhanced Input|User Settings", meta=(DisplayName="Enable User Settings (Experimental)", DisplayPriority = 1))
 	uint8 bEnableUserSettings : 1;
 
 	/** If true, then the DefaultMappingContexts will be applied to all Enhanced Input Subsystems. */
