Serialize all customized keys, not just dirty ones. This fixes a bug where if you remap a key, save, exit, remap a different key, and save, then the first key wasn't being saved

#rnx

diff --git a/Engine/Plugins/EnhancedInput/Source/EnhancedInput/Private/UserSettings/EnhancedInputUserSettings.cpp b/Engine/Plugins/EnhancedInput/Source/EnhancedInput/Private/UserSettings/EnhancedInputUserSettings.cpp
--- a/Engine/Plugins/EnhancedInput/Source/EnhancedInput/Private/UserSettings/EnhancedInputUserSettings.cpp
+++ b/Engine/Plugins/EnhancedInput/Source/EnhancedInput/Private/UserSettings/EnhancedInputUserSettings.cpp
@@ -38,6 +38,7 @@ namespace UE::EnhancedInput
 	UE_DEFINE_GAMEPLAY_TAG_STATIC(TAG_InvalidMappingName, "InputUserSettings.FailureReasons.InvalidMappingName");
 	UE_DEFINE_GAMEPLAY_TAG_STATIC(TAG_NoKeyProfile, "InputUserSettings.FailureReasons.NoKeyProfile");
 	UE_DEFINE_GAMEPLAY_TAG_STATIC(TAG_NoMatchingMappings, "InputUserSettings.FailureReasons.NoMatchingMappings");
+	UE_DEFINE_GAMEPLAY_TAG_STATIC(TAG_NoMappingRowFound, "InputUserSettings.FailureReasons.NoMappingRowFound");
 
 	static void DumpAllKeyProfilesToLog(const TArray<FString>& Args)
 	{
@@ -709,6 +710,9 @@ namespace UE::EnhancedInput
 
 void UEnhancedInputUserSettings::Serialize(FArchive& Ar)
 {
+	// We always want to ensure that we call the Super::Serialize function here because
+	// not doing that can introduce many hard to track down serialization bugs and not include
+	// the UObject info of this settings object.
 	Super::Serialize(Ar);
 
 	if (IsTemplate() || Ar.IsCountingMemory())
@@ -742,7 +746,8 @@ void UEnhancedInputUserSettings::Serialize(FArchive& Ar)
 			{
 				for (FPlayerKeyMapping& Mapping : MappingRow.Value.Mappings)
 				{
-					if (Mapping.IsDirty())
+					// We want to save any dirty or customized key mappings
+					if (Mapping.IsDirty() || Mapping.IsCustomized())
 					{
 						Header.DirtyMappings.Push(
 							{
@@ -751,6 +756,10 @@ void UEnhancedInputUserSettings::Serialize(FArchive& Ar)
 								/* .HardwareDeviceId = */ Mapping.HardwareDeviceId, 
 								/* .Slot = */ Mapping.Slot
 							});
+						
+						// Since we are about to save this key mapping, we no longer need to treat
+						// it as dirty
+						Mapping.bIsDirty = false;
 					}
 				}
 			}
@@ -787,10 +796,13 @@ void UEnhancedInputUserSettings::Serialize(FArchive& Ar)
 		}
 	}
 
+	// Because we are saving subobjects here, we need to serialize their UObject information as well
 	FString SavedObjectTerminator = TEXT("ObjectEnd");
 	
 	for (TPair<FGameplayTag, UEnhancedPlayerMappableKeyProfile*> ProfilePair : SavedKeyProfiles)
 	{
+		// FYI: This call to serialize will only be serializing the basic UObject info of this key profile,
+		// not all the individual key mappings :) 
 		ProfilePair.Value->Serialize(Ar);
 
 		// Save a terminator after each subobject
@@ -830,6 +842,8 @@ void UEnhancedInputUserSettings::MapPlayerKey(const FMapPlayerKeyArgs& InArgs, F
 	{
 		// Then set the player mapped key
 		FoundMapping->SetCurrentKey(InArgs.NewKey);
+		OnKeyMappingUpdated(FoundMapping, InArgs, false);
+		
 		OnSettingsChanged.Broadcast(this);
 	}
 	// If it doesn't exist, then we need to make it if there is a valid action name
@@ -848,17 +862,22 @@ void UEnhancedInputUserSettings::MapPlayerKey(const FMapPlayerKeyArgs& InArgs, F
 			FPlayerKeyMapping PlayerMappingData = {};
 			PlayerMappingData.MappingName = InArgs.MappingName;
 			PlayerMappingData.Slot = InArgs.Slot;
+			// This is a new mapping, and should be dirty upon creation
+			PlayerMappingData.bIsDirty = true;
 
-			// If there is some valid hardware then keep track of that
-			if (const UInputPlatformSettings* PlatformSettings = UInputPlatformSettings::Get())
+			// Check for known hardware device Id's if one has been specified
+			if (!InArgs.HardwareDeviceId.IsNone())
 			{
-				if (const FHardwareDeviceIdentifier* Hardware = PlatformSettings->GetHardwareDeviceForClassName(InArgs.HardwareDeviceId))
-				{
-					PlayerMappingData.HardwareDeviceId = *Hardware;	
-				}
-				else
+				if (const UInputPlatformSettings* PlatformSettings = UInputPlatformSettings::Get())
 				{
-					UE_LOG(LogEnhancedInput, Log, TEXT("[UEnhancedInputUserSettings::MapPlayerKey] Unable to find a matching Hardware Device Identifier with the HardwareDeviceId of '%s'"), *InArgs.HardwareDeviceId.ToString());
+					if (const FHardwareDeviceIdentifier* Hardware = PlatformSettings->GetHardwareDeviceForClassName(InArgs.HardwareDeviceId))
+					{
+						PlayerMappingData.HardwareDeviceId = *Hardware;	
+					}
+					else
+					{
+						UE_LOG(LogEnhancedInput, Log, TEXT("[UEnhancedInputUserSettings::MapPlayerKey] Unable to find a matching Hardware Device Identifier with the HardwareDeviceId of '%s'"), *InArgs.HardwareDeviceId.ToString());
+					}
 				}
 			}
 			
@@ -866,10 +885,19 @@ void UEnhancedInputUserSettings::MapPlayerKey(const FMapPlayerKeyArgs& InArgs, F
 			// EKeys::Invalid and we only need to track the player mapped key
 			PlayerMappingData.SetCurrentKey(InArgs.NewKey);
 			PlayerMappingData.DisplayName = ExistingMapping->DisplayName;
-			MappingRow->Mappings.Add(PlayerMappingData);
+			
+			const FSetElementId SetElem = MappingRow->Mappings.Add(PlayerMappingData);
+			OnKeyMappingUpdated(&MappingRow->Mappings.Get(SetElem), InArgs, false);
+			
 			OnSettingsChanged.Broadcast(this);
 		}
 	}
+	else
+	{
+		// If there is neither an existing key mapping or a key mapping row, then mapping the key has failed
+		FailureReason.AddTag(UE::EnhancedInput::TAG_NoMappingRowFound);
+		UE_LOG(LogEnhancedInput, Warning, TEXT("[UEnhancedInputUserSettings::MapPlayerKey] Failed to map a player key for '%s'"), *InArgs.MappingName.ToString());
+	}
 }
 
 void UEnhancedInputUserSettings::UnMapPlayerKey(const FMapPlayerKeyArgs& InArgs, FGameplayTagContainer& FailureReason)
@@ -896,6 +924,7 @@ void UEnhancedInputUserSettings::UnMapPlayerKey(const FMapPlayerKeyArgs& InArgs,
 		// The settings have only changed if the mapping is dirty now
 		if (FoundMapping->IsDirty())
 		{
+			OnKeyMappingUpdated(FoundMapping, InArgs, true);
 			OnSettingsChanged.Broadcast(this);
 		}
 		
@@ -908,6 +937,11 @@ void UEnhancedInputUserSettings::UnMapPlayerKey(const FMapPlayerKeyArgs& InArgs,
 	}
 }
 
+void UEnhancedInputUserSettings::OnKeyMappingUpdated(FPlayerKeyMapping* ChangedMapping, const FMapPlayerKeyArgs& InArgs, const bool bIsBeingUnmapped)
+{
+	// Do nothing by default
+}
+
 const TSet<FPlayerKeyMapping>& UEnhancedInputUserSettings::FindMappingsInRow(const FName MappingName) const
 {
 	if (UEnhancedPlayerMappableKeyProfile* KeyProfile = GetCurrentKeyProfile())
@@ -962,6 +996,13 @@ bool UEnhancedInputUserSettings::SetKeyProfile(const FGameplayTag& InProfileId)
 
 	const FGameplayTag OriginalProfileId = CurrentProfileIdentifier;
 	UEnhancedPlayerMappableKeyProfile* OriginalProfile = GetCurrentKeyProfile();
+
+	// If this key profile is already equipped, then there is nothing to do
+	if (InProfileId == CurrentProfileIdentifier)
+	{
+		UE_LOG(LogEnhancedInput, Log, TEXT("Key profile '%s' is already currently equipped. Nothing will happen."), *InProfileId.ToString());
+		return false;
+	}
 	
 	if (const TObjectPtr<UEnhancedPlayerMappableKeyProfile>* NewProfile = SavedKeyProfiles.Find(InProfileId))
 	{
@@ -977,6 +1018,9 @@ bool UEnhancedInputUserSettings::SetKeyProfile(const FGameplayTag& InProfileId)
 		// Keep track of what the current profile is now
 		CurrentProfileIdentifier = InProfileId;
 
+		// TODO: Register all the input mapping contexts with the new key profile to
+		// ensure that it has all the default key mappings
+
 		// Let any listeners know that the mapping profile has changed
 		OnKeyProfileChanged.Broadcast(NewProfile->Get());
 	}
diff --git a/Engine/Plugins/EnhancedInput/Source/EnhancedInput/Public/UserSettings/EnhancedInputUserSettings.h b/Engine/Plugins/EnhancedInput/Source/EnhancedInput/Public/UserSettings/EnhancedInputUserSettings.h
--- a/Engine/Plugins/EnhancedInput/Source/EnhancedInput/Public/UserSettings/EnhancedInputUserSettings.h
+++ b/Engine/Plugins/EnhancedInput/Source/EnhancedInput/Public/UserSettings/EnhancedInputUserSettings.h
@@ -496,18 +496,31 @@ protected:
 
 	virtual void OnAsyncSaveComplete(const FString& SlotName, const int32 UserIndex, bool bSuccess);
 
+	/**
+	 * A virtual function that will be called whenever a player key mapping has been updated, before OnSettingsChanged is called.
+	 * This provides subclasses the opportunity to modify this key mapping before any settings screens are updated
+	 *
+	 * This does nothing by default, but some subclasses may want to modify the mapping. 
+	 * Perhaps to change to a specific hardware device ID or something like that
+	 *
+	 * @param ChangedMapping	The key mapping that has changed
+	 * @param InArgs			The arguments that were used to change this mapping.
+	 * @param bIsBeingUnmapped	True if this key mapping is being unmapped. False if the key mapping is just being changed
+	 */
+	virtual void OnKeyMappingUpdated(FPlayerKeyMapping* ChangedMapping, const FMapPlayerKeyArgs& InArgs, const bool bIsBeingUnmapped);
+
 public:
 	
 	ULocalPlayer* GetLocalPlayer() const;
 	
 	/** Fired when the user settings have changed, such as their key mappings. */
 	DECLARE_DYNAMIC_MULTICAST_DELEGATE_OneParam(FEnhancedInputUserSettingsChanged, UEnhancedInputUserSettings*, Settings);
-	UPROPERTY(BlueprintAssignable, Category = "Enhanced Input|User Settings")
+	UPROPERTY(BlueprintAssignable, Transient, Category = "Enhanced Input|User Settings")
 	FEnhancedInputUserSettingsChanged OnSettingsChanged;
 	
 	/** Called after the settings have been applied from the ApplySettings call. */
 	DECLARE_DYNAMIC_MULTICAST_DELEGATE(FEnhancedInputUserSettingsApplied);
-	UPROPERTY(BlueprintAssignable, Category = "Enhanced Input|User Settings")
+	UPROPERTY(BlueprintAssignable, Transient, Category = "Enhanced Input|User Settings")
 	FEnhancedInputUserSettingsApplied OnSettingsApplied;
 
 	// Remappable keys API
