Register IMC's to all saved key profiles instead of just the current only. Break out the registration of keys to a different function to allow us to re-register IMC's  when on serialize to ensure the mappings are up to date

#rnx

diff --git a/Engine/Plugins/EnhancedInput/Source/EnhancedInput/Private/UserSettings/EnhancedInputUserSettings.cpp b/Engine/Plugins/EnhancedInput/Source/EnhancedInput/Private/UserSettings/EnhancedInputUserSettings.cpp
--- a/Engine/Plugins/EnhancedInput/Source/EnhancedInput/Private/UserSettings/EnhancedInputUserSettings.cpp
+++ b/Engine/Plugins/EnhancedInput/Source/EnhancedInput/Private/UserSettings/EnhancedInputUserSettings.cpp
@@ -777,8 +777,10 @@ void UEnhancedInputUserSettings::Serialize(FArchive& Ar)
 			if (const UClass* FoundClass = FindObject<UClass>(nullptr, *Header.ClassPath))
 			{
 				UEnhancedPlayerMappableKeyProfile* NewProfile = NewObject<UEnhancedPlayerMappableKeyProfile>(/* outer */ this, /* class */ FoundClass);
-				SavedKeyProfiles.Add(Header.ProfileIdentifier, NewProfile);
 
+				// Add the new profile to the settings object. This will replace the old profile if one existed
+				SavedKeyProfiles.Add(Header.ProfileIdentifier, NewProfile);				
+				
 				// Add any player saved keys to this profile!
 				for (UE::EnhancedInput::FKeyMappingSaveData& SavedKeyData : Header.DirtyMappings)
 				{
@@ -792,6 +794,18 @@ void UEnhancedInputUserSettings::Serialize(FArchive& Ar)
 					
 					MappingRow.Mappings.Add(PlayerMapping);
 				}
+
+				// We need to populate this key profile with all the known key mappings so that it's up to date 
+				for (TObjectPtr<UInputMappingContext> IMC : RegisteredMappingContexts)
+				{
+					RegisterKeyMappingsToProfile(*NewProfile, IMC);
+				}
+
+				// Ensure that the owning local player is up to date
+				if (const ULocalPlayer* LP = GetLocalPlayer())
+				{
+					NewProfile->OwningUserId = LP->GetPlatformUserId();
+				}
 			}
 		}
 	}
@@ -1082,6 +1096,12 @@ UEnhancedPlayerMappableKeyProfile* UEnhancedInputUserSettings::CreateNewKeyProfi
 		OutProfile->OwningUserId = InArgs.UserId;
 		
 		SavedKeyProfiles.Add(InArgs.ProfileIdentifier, OutProfile);
+		
+		// We need to populate this key profile with all the known key mappings so that it's up to date 
+		for (TObjectPtr<UInputMappingContext> IMC : RegisteredMappingContexts)
+		{
+			RegisterKeyMappingsToProfile(*OutProfile, IMC);
+		}
 	}
 	
 	// set as current
@@ -1125,19 +1145,31 @@ bool UEnhancedInputUserSettings::RegisterInputMappingContext(UInputMappingContex
 		return false;
 	}
 
-	// There is no need to re-register an IMC
+	// There is no need to re-register an IMC if it is has already been registered.
 	if (RegisteredMappingContexts.Contains(IMC))
 	{
-		UE_LOG(LogEnhancedInput, VeryVerbose, TEXT("Mapping Context '%s' is already registered with the User Settings."), *IMC->GetFName().ToString());
 		return false;
 	}
 	
-	const FSetElementId Res = RegisteredMappingContexts.Add(IMC);
+	// Keep track of all the registered IMC's
+	RegisteredMappingContexts.Add(IMC);
+
+	bool bResult = true;
+
+	// Register the mappings of this IMC to every saved key profile
+	for (TPair<FGameplayTag, TObjectPtr<UEnhancedPlayerMappableKeyProfile>> Pair : SavedKeyProfiles)
+	{
+		bResult &= RegisterKeyMappingsToProfile(*Pair.Value, IMC);
+	}
+	
+	return bResult;
+}
 
-	UEnhancedPlayerMappableKeyProfile* CurrentProfile = GetCurrentKeyProfile();
-	if (!CurrentProfile)
+bool UEnhancedInputUserSettings::RegisterKeyMappingsToProfile(UEnhancedPlayerMappableKeyProfile& Profile, UInputMappingContext* IMC)
+{
+	if (!IMC)
 	{
-		UE_LOG(LogEnhancedInput, Error, TEXT("There is not an active key profile!"));
+		UE_LOG(LogEnhancedInput, Error, TEXT("Attempting to register a null mapping context with the user settings!"));
 		ensure(false);
 		return false;
 	}
@@ -1164,7 +1196,7 @@ bool UEnhancedInputUserSettings::RegisterInputMappingContext(UInputMappingContex
 		int32& NumDefinedMappings = OriginalMappingsDesiredSlot.FindOrAdd(ActionName);
 
 		// Find or create a mapping row 
-		FKeyMappingRow& MappingRow = CurrentProfile->PlayerMappedKeys.FindOrAdd(ActionName);
+		FKeyMappingRow& MappingRow = Profile.PlayerMappedKeys.FindOrAdd(ActionName);
 		
 		// By default, the slot will be determined by how many mappings this action has already.
 		// So if this is the first default mapping, then this will be EPlayerMappableKeySlot::First,
@@ -1208,14 +1240,11 @@ bool UEnhancedInputUserSettings::RegisterInputMappingContext(UInputMappingContex
 			MappingRow.Mappings.Add({ KeyMapping, IMCDefinedSlot });	
 		}
 	}
-
-	if (Res.IsValidId())
-	{
-		OnMappingContextRegistered.Broadcast(IMC);	
-	}
+	
+	OnMappingContextRegistered.Broadcast(IMC);	
 
 	UE_LOG(LogEnhancedInput, Verbose, TEXT("Registered IMC with UEnhancedInputUserSettings: %s"), *IMC->GetFName().ToString());
-	return Res.IsValidId();
+	return true;
 }
 
 bool UEnhancedInputUserSettings::UnregisterInputMappingContext(const UInputMappingContext* IMC)
diff --git a/Engine/Plugins/EnhancedInput/Source/EnhancedInput/Public/UserSettings/EnhancedInputUserSettings.h b/Engine/Plugins/EnhancedInput/Source/EnhancedInput/Public/UserSettings/EnhancedInputUserSettings.h
--- a/Engine/Plugins/EnhancedInput/Source/EnhancedInput/Public/UserSettings/EnhancedInputUserSettings.h
+++ b/Engine/Plugins/EnhancedInput/Source/EnhancedInput/Public/UserSettings/EnhancedInputUserSettings.h
@@ -632,6 +632,17 @@ public:
 
 protected:
 
+	/**
+	 * This will register individual key mappings from the given Input Mapping context on the profile id.
+	 *
+	 * If the mappings exist on the profile already, then they will be updated based on the latest data from the IMC
+	 * otherwise, they will be added as new FPlayerKeyMapping.
+	 *
+	 * This will be called for each key profile when users initially call RegisterInputMappingContext, and on
+	 * load during Serialize to correctly populate a loaded profile with all the known registered IMC's mappings.
+	 */
+	virtual bool RegisterKeyMappingsToProfile(UEnhancedPlayerMappableKeyProfile& Profile, UInputMappingContext* IMC);
+
 	/** 
 	* Provides a space for subclasses to add additional debug info if desired.
 	* By default this will draw info about all player mapped keys.
