Only broadcast the OnMappingContextRegistered delegate in the settings if it was successfully registered to all key profiles, not each time a mapping is added. Keep track of the hardware device that each mapping is in. This will allow you to have each hardware device in a unique slot!

#rnx

diff --git a/Engine/Plugins/EnhancedInput/Source/EnhancedInput/Private/UserSettings/EnhancedInputUserSettings.cpp b/Engine/Plugins/EnhancedInput/Source/EnhancedInput/Private/UserSettings/EnhancedInputUserSettings.cpp
--- a/Engine/Plugins/EnhancedInput/Source/EnhancedInput/Private/UserSettings/EnhancedInputUserSettings.cpp
+++ b/Engine/Plugins/EnhancedInput/Source/EnhancedInput/Private/UserSettings/EnhancedInputUserSettings.cpp
@@ -236,14 +236,20 @@ void FPlayerKeyMapping::SetHardwareDeviceId(const FHardwareDeviceIdentifier& InD
 	HardwareDeviceId = InDeviceId;
 }
 
-void FPlayerKeyMapping::UpdateOriginalKey(const FEnhancedActionKeyMapping& OriginalMapping)
+void FPlayerKeyMapping::UpdateMetadataFromActionKeyMapping(const FEnhancedActionKeyMapping& OriginalMapping)
 {
 	ensure(OriginalMapping.IsPlayerMappable() && OriginalMapping.GetMappingName() == MappingName);
 	
-	DefaultKey = OriginalMapping.Key;
 	DisplayCategory = OriginalMapping.GetDisplayCategory();
 	DisplayName = OriginalMapping.GetDisplayName();
 	AssociatedInputAction = OriginalMapping.Action;
+}
+
+void FPlayerKeyMapping::UpdateDefaultKeyFromActionKeyMapping(const FEnhancedActionKeyMapping& OriginalMapping)
+{
+	ensure(OriginalMapping.IsPlayerMappable() && OriginalMapping.GetMappingName() == MappingName);
+
+	DefaultKey = OriginalMapping.Key;
 
 	// If the mapping is the same as the default key, make sure that it is not marked as
 	// dirty so that we don't serialize it
@@ -951,6 +957,44 @@ void UEnhancedInputUserSettings::UnMapPlayerKey(const FMapPlayerKeyArgs& InArgs,
 	}
 }
 
+void UEnhancedInputUserSettings::ResetAllPlayerKeysInRow(const FMapPlayerKeyArgs& InArgs, FGameplayTagContainer& FailureReason)
+{
+	if (!InArgs.MappingName.IsValid())
+	{
+		FailureReason.AddTag(UE::EnhancedInput::TAG_InvalidMappingName);
+		return;
+	}
+
+	// Get the key profile that was specified
+	UEnhancedPlayerMappableKeyProfile* KeyProfile = InArgs.ProfileId.IsValid() ? GetKeyProfileWithIdentifier(InArgs.ProfileId) : GetCurrentKeyProfile();
+	if (!KeyProfile)
+	{
+		FailureReason.AddTag(UE::EnhancedInput::TAG_NoKeyProfile);
+		return;
+	}
+
+	if (FKeyMappingRow* ExistingMappings = KeyProfile->PlayerMappedKeys.Find(InArgs.MappingName))
+	{
+		for (FPlayerKeyMapping& Mapping : ExistingMappings->Mappings)
+		{
+			// Then set the player mapped key
+			Mapping.ResetToDefault();
+
+			// The settings have only changed if the mapping is dirty now
+			if (Mapping.IsDirty())
+			{
+				OnKeyMappingUpdated(&Mapping, InArgs, true);
+				OnSettingsChanged.Broadcast(this);
+			}
+			UE_LOG(LogEnhancedInput, Verbose, TEXT("[UEnhancedInputUserSettings::MapPlayerKey] Reset keymapping to default: '%s'"), *Mapping.ToString());
+		}
+	}
+	else
+	{
+		FailureReason.AddTag(UE::EnhancedInput::TAG_NoMatchingMappings);
+	}
+}
+
 void UEnhancedInputUserSettings::OnKeyMappingUpdated(FPlayerKeyMapping* ChangedMapping, const FMapPlayerKeyArgs& InArgs, const bool bIsBeingUnmapped)
 {
 	// Do nothing by default
@@ -1130,7 +1174,7 @@ bool UEnhancedInputUserSettings::RegisterInputMappingContexts(const TSet<UInputM
 
 	for (UInputMappingContext* IMC : MappingContexts)
 	{
-		bResult |= RegisterInputMappingContext(IMC);
+		bResult &= RegisterInputMappingContext(IMC);
 	}
 
 	return bResult;
@@ -1161,6 +1205,11 @@ bool UEnhancedInputUserSettings::RegisterInputMappingContext(UInputMappingContex
 	{
 		bResult &= RegisterKeyMappingsToProfile(*Pair.Value, IMC);
 	}
+
+	if (bResult)
+	{
+		OnMappingContextRegistered.Broadcast(IMC);
+	}
 	
 	return bResult;
 }
@@ -1174,10 +1223,9 @@ bool UEnhancedInputUserSettings::RegisterKeyMappingsToProfile(UEnhancedPlayerMap
 		return false;
 	}
 
-	// A map that is used to store how many actions of a certain name are in this IMC
-	// this is how we can determine the slot that a new player key mapping needs
-	static TMap<FName, int32> OriginalMappingsDesiredSlot;
-	OriginalMappingsDesiredSlot.Reset();
+	typedef TMap<FName, EPlayerMappableKeySlot> FSlotCountToMappingNames;
+	static TMap<FHardwareDeviceIdentifier, FSlotCountToMappingNames> MappingsPerDeviceTypes;
+	MappingsPerDeviceTypes.Reset();
 	
 	for (const FEnhancedActionKeyMapping& KeyMapping : IMC->GetMappings())
 	{
@@ -1187,49 +1235,45 @@ bool UEnhancedInputUserSettings::RegisterKeyMappingsToProfile(UEnhancedPlayerMap
 			continue;
 		}
 
-		// Get the unique FName for the "Action name". This is set on the UInputAction
-		// but can be overriden on each FEnhancedActionKeyMapping to create multiple
+		// Get the unique FName for the "Mapping name". This is set on the UInputAction
+		// but can be overridden on each FEnhancedActionKeyMapping to create multiple
 		// mapping options for a single Input Action. 
-		const FName ActionName = KeyMapping.GetMappingName();
-
-		// Iterate the number of mappings that use this action
-		int32& NumDefinedMappings = OriginalMappingsDesiredSlot.FindOrAdd(ActionName);
+		const FName MappingName = KeyMapping.GetMappingName();
 
 		// Find or create a mapping row 
-		FKeyMappingRow& MappingRow = Profile.PlayerMappedKeys.FindOrAdd(ActionName);
+		FKeyMappingRow& MappingRow = Profile.PlayerMappedKeys.FindOrAdd(MappingName);
+
+		// Determine what Slot and Hardware device this Action Mapping should go in
+		const FHardwareDeviceIdentifier MappingDeviceType = DetermineHardwareDeviceForActionMapping(KeyMapping);
 		
-		// By default, the slot will be determined by how many mappings this action has already.
-		// So if this is the first default mapping, then this will be EPlayerMappableKeySlot::First,
-		// if this is the second element going into this set then it will be EPlayerMappableKeySlot::Second
-		// and so on
-		const EPlayerMappableKeySlot IMCDefinedSlot =
-			(EPlayerMappableKeySlot)
-			(FMath::Min<uint8>(static_cast<uint8>(NumDefinedMappings), static_cast<uint8>(EPlayerMappableKeySlot::Max)));
-
-		// Update the number of mappings that have this action name
-		++NumDefinedMappings;
+		FSlotCountToMappingNames& MappingNameToSlotCount = MappingsPerDeviceTypes.FindOrAdd(MappingDeviceType);
+
+		// Determine what slot is available for the mapping to be placed in
+		EPlayerMappableKeySlot& MappingSlot = MappingNameToSlotCount.FindOrAdd(MappingName);
 		
 		bool bUpdatedExistingMapping = false;
 		
-		// If a key mapping exists when a mapping context is being registered, that means that it is a 
-		// customized player mapped key.
+		// Iterate any existing mappings in this row to ensure that the metadata and default key values are correct.
+		// At this stage, keys will only exist here if they are player customized
 		for (FPlayerKeyMapping& ExistingMapping : MappingRow.Mappings)
 		{
-			// We only want to update exiting mapping in the desired slot if one existed in the mapping context.
-			if (ExistingMapping.GetSlot() == IMCDefinedSlot)
+			// We only want to update the default _key_ for an existing mapping if it is from
+			// the same slot as this Action Mapping would be going in.
+			// This is because players can map keys to slots are not defined in the input mapping context,
+			// thus their default key should be EKeys::Invalid.
+			// For example, you define a single key mapping of "Jump" to space bar in your IMC.
+			// The mapping to spacebar is in EPlayerMappableKeySlot::First. The player adds a second
+			// key mapping of "Jump" to the X key, and keeps the spacebar mapping. That is in EPlayerMappableKeySlot::Second
+			// We should no update the default key of second slot because it's default value should be EKeys::Invalid.
+			if (ExistingMapping.GetSlot() == MappingSlot && ExistingMapping.HardwareDeviceId.PrimaryDeviceType == MappingDeviceType.PrimaryDeviceType)
 			{
-				// Update the fields on this mapping that the player can't change so that they
-				// get updated if the source UInputMappingContext asset has been changed
-				ExistingMapping.UpdateOriginalKey(KeyMapping);
+				ExistingMapping.UpdateDefaultKeyFromActionKeyMapping(KeyMapping);
 				bUpdatedExistingMapping = true;
 			}
 			
-			// it is possible that a player has mapped to more slots then you have specified
-			// in a mapping context asset. In that case, they won't have display names populated
-			// and their default key is correctly set to EKeys::Invalid
-			// That is why we need to move the display name to the mapping ROW, since the actions
-			// will have the same display name no matter what.
-			ExistingMapping.DisplayName = KeyMapping.GetDisplayName();
+			// We always want to ensure that the metadata on a player key mapping is up to date
+			// from the Input Mapping Context.
+			ExistingMapping.UpdateMetadataFromActionKeyMapping(KeyMapping);
 		}
 
 		// If the mapping was not found in the existing row, then the player has not mapped anything to it.
@@ -1237,16 +1281,22 @@ bool UEnhancedInputUserSettings::RegisterKeyMappingsToProfile(UEnhancedPlayerMap
 		if (!bUpdatedExistingMapping)
 		{
 			// Add a default mapping to this row
-			MappingRow.Mappings.Add({ KeyMapping, IMCDefinedSlot });	
+			MappingRow.Mappings.Add({ KeyMapping, MappingSlot, MappingDeviceType });
 		}
+
+		// Increment the mapping slot to keep track of how many slots are in use by which hardware device
+		MappingSlot = static_cast<EPlayerMappableKeySlot>(FMath::Min<uint8>(static_cast<uint8>(MappingSlot) + 1, static_cast<uint8>(EPlayerMappableKeySlot::Max)));
 	}
-	
-	OnMappingContextRegistered.Broadcast(IMC);	
 
 	UE_LOG(LogEnhancedInput, Verbose, TEXT("Registered IMC with UEnhancedInputUserSettings: %s"), *IMC->GetFName().ToString());
 	return true;
 }
 
+FHardwareDeviceIdentifier UEnhancedInputUserSettings::DetermineHardwareDeviceForActionMapping(const FEnhancedActionKeyMapping& ActionMapping) const
+{
+	return FHardwareDeviceIdentifier::Invalid;
+}
+
 bool UEnhancedInputUserSettings::UnregisterInputMappingContext(const UInputMappingContext* IMC)
 {
 	return RegisteredMappingContexts.Remove(IMC) != INDEX_NONE;
diff --git a/Engine/Plugins/EnhancedInput/Source/EnhancedInput/Public/UserSettings/EnhancedInputUserSettings.h b/Engine/Plugins/EnhancedInput/Source/EnhancedInput/Public/UserSettings/EnhancedInputUserSettings.h
--- a/Engine/Plugins/EnhancedInput/Source/EnhancedInput/Public/UserSettings/EnhancedInputUserSettings.h
+++ b/Engine/Plugins/EnhancedInput/Source/EnhancedInput/Public/UserSettings/EnhancedInputUserSettings.h
@@ -152,12 +152,18 @@ public:
 	/** Sets the value of the hardware device ID that this custom key mapping is associated with. This can be used to filter key mappings per input device */
 	void SetHardwareDeviceId(const FHardwareDeviceIdentifier& InDeviceId);
 
+	/**
+	 * Sets the Default Key on this mapping to that of the given ActionKeyMapping.
+	 * Checks for customization of this key. This is for use during player key registration.
+	 */
+	void UpdateDefaultKeyFromActionKeyMapping(const FEnhancedActionKeyMapping& OriginalMapping);
+
 	/**
 	 * Updates the metadata properties on this player mapped key based on the given
 	 * enhanced action mapping. This will populate the fields on this struct that are not editable
 	 * by the player such as the localized display name and default key.  
 	 */
-	void UpdateOriginalKey(const FEnhancedActionKeyMapping& OriginalMapping);
+	void UpdateMetadataFromActionKeyMapping(const FEnhancedActionKeyMapping& OriginalMapping);
 	
 	ENHANCEDINPUT_API friend uint32 GetTypeHash(const FPlayerKeyMapping& InMapping);
 	bool operator==(const FPlayerKeyMapping& Other) const;
@@ -165,7 +171,7 @@ public:
 
 	/** Returns true if this mapping has been modified since it was registered from an IMC */
 	const bool IsDirty() const;
-
+	
 protected:
 	
 	/** The name of the mapping for this key */
@@ -531,10 +537,22 @@ public:
 	UFUNCTION(BlueprintCallable, Category="Enhanced Input|User Settings", meta = (AutoCreateRefTerm = "FailureReason"))
 	virtual void MapPlayerKey(const FMapPlayerKeyArgs& InArgs, FGameplayTagContainer& FailureReason);
 
-	/** Unmap what is currently mapped to the given mapping in the given slot */
+	/** 
+	* Unmaps a single player mapping that matches the given Mapping name, slot, and hardware device.
+	*/
 	UFUNCTION(BlueprintCallable, Category="Enhanced Input|User Settings", meta = (AutoCreateRefTerm = "FailureReason"))
 	virtual void UnMapPlayerKey(const FMapPlayerKeyArgs& InArgs, FGameplayTagContainer& FailureReason);
 
+	/**
+	* Resets each player mapped key to it's default value from the Input Mapping Context that it was registered from.
+	* If a key did not come from an IMC (i.e. it was added additionally by the player) then it will be reset to EKeys::Invalid.
+	* 
+	* @param InArgs				Arguments that contain the mapping name and profile ID to find the mapping to reset.
+	* @param FailureReason		Populated with failure reasons if the operation fails.
+	*/
+	UFUNCTION(BlueprintCallable, Category="Enhanced Input|User Settings", meta = (AutoCreateRefTerm = "FailureReason"))
+	virtual void ResetAllPlayerKeysInRow(const FMapPlayerKeyArgs& InArgs, FGameplayTagContainer& FailureReason);
+
 	/** Returns a set of all player key mappings for the given mapping name. */
 	UFUNCTION(BlueprintCallable, Category="Enhanced Input|User Settings")
 	virtual const TSet<FPlayerKeyMapping>& FindMappingsInRow(const FName MappingName) const;
@@ -643,6 +661,15 @@ protected:
 	 */
 	virtual bool RegisterKeyMappingsToProfile(UEnhancedPlayerMappableKeyProfile& Profile, UInputMappingContext* IMC);
 
+	/**
+	 * Determines the hardware device that an action key mapping is associated with. By default, the hardware will be FHardwareDeviceIdentifier::Invalid.
+	 *
+	 * Override this function if you wish to have Input Device Specific Player Key mappings. For example, specify hardware
+	 * devices based on some metadata on the UPlayerMappableKeySettings of the action mapping, or infer a device based
+	 * on the type of FKey it is.
+	 */
+	virtual FHardwareDeviceIdentifier DetermineHardwareDeviceForActionMapping(const FEnhancedActionKeyMapping& ActionMapping) const;
+	
 	/** 
 	* Provides a space for subclasses to add additional debug info if desired.
 	* By default this will draw info about all player mapped keys.
