Specify what asset is using the mapping name if one exists already, making it much easier to track down in use mappings :)

Add in a cvar 'EnhancedInput.Editor.EnableMappingNameValidation' to toggle name validation, just to make sure that we can do sanity checks that the validation is actually accurate

#rnx

diff --git a/Engine/Plugins/EnhancedInput/Source/InputEditor/Private/EnhancedInputPlayerMappableNameValidator.cpp b/Engine/Plugins/EnhancedInput/Source/InputEditor/Private/EnhancedInputPlayerMappableNameValidator.cpp
--- a/Engine/Plugins/EnhancedInput/Source/InputEditor/Private/EnhancedInputPlayerMappableNameValidator.cpp
+++ b/Engine/Plugins/EnhancedInput/Source/InputEditor/Private/EnhancedInputPlayerMappableNameValidator.cpp
@@ -1,7 +1,22 @@
 // Copyright Epic Games, Inc. All Rights Reserved.
 
 #include "EnhancedInputPlayerMappableNameValidator.h"
+#include "PlayerMappableKeySettings.h"
 #include "InputEditorModule.h"
+#include "Internationalization/Text.h"	// For FFormatNamedArguments
+#include "HAL/IConsoleManager.h"		// For FAutoConsoleVariableRef
+
+namespace UE::EnhancedInput
+{
+	/** Enables editor validation on player mapping names */
+	static bool bEnableMappingNameValidation = true;
+	
+	FAutoConsoleVariableRef CVarEnableMappingNameValidation(
+	TEXT("EnhancedInput.Editor.EnableMappingNameValidation"),
+	bEnableMappingNameValidation,
+	TEXT("Enables editor validation on player mapping names"),
+	ECVF_Default);
+}
 
 FEnhancedInputPlayerMappableNameValidator::FEnhancedInputPlayerMappableNameValidator(FName InExistingName)
 	: FStringSetNameValidator(InExistingName.ToString())
@@ -11,10 +26,29 @@ EValidatorResult FEnhancedInputPlayerMappableNameValidator::IsValid(const FStrin
 {
 	EValidatorResult Result = FStringSetNameValidator::IsValid(Name, bOriginal);
 
-	if (Result != EValidatorResult::ExistingName && FInputEditorModule::IsMappingNameInUse(FName(Name)))
+	if (UE::EnhancedInput::bEnableMappingNameValidation &&
+		Result != EValidatorResult::ExistingName &&
+		FInputEditorModule::IsMappingNameInUse(FName(Name)))
 	{
 		Result = EValidatorResult::AlreadyInUse;
 	}
 
 	return Result;
+}
+
+FText FEnhancedInputPlayerMappableNameValidator::GetErrorText(const FString& Name, EValidatorResult ErrorCode)
+{
+	// Attempt to specify what asset is using this name
+	if (ErrorCode == EValidatorResult::AlreadyInUse)
+	{
+		if (const UPlayerMappableKeySettings* Settings = FInputEditorModule::FindMappingByName(FName(Name)))
+		{
+			FFormatNamedArguments Args;
+			Args.Add(TEXT("AssetUsingName"), FText::FromString(GetNameSafe(Settings->GetOuter())));
+
+			return FText::Format(NSLOCTEXT("EnhancedInput", "MappingNameInUseBy_Error", "Name is already in use by '{AssetUsingName}'"), Args);	
+		}
+	}
+	
+	return INameValidatorInterface::GetErrorText(Name, ErrorCode);
 }
\ No newline at end of file
diff --git a/Engine/Plugins/EnhancedInput/Source/InputEditor/Private/InputCustomizations.cpp b/Engine/Plugins/EnhancedInput/Source/InputEditor/Private/InputCustomizations.cpp
--- a/Engine/Plugins/EnhancedInput/Source/InputEditor/Private/InputCustomizations.cpp
+++ b/Engine/Plugins/EnhancedInput/Source/InputEditor/Private/InputCustomizations.cpp
@@ -659,7 +659,7 @@ bool FPlayerMappableKeyChildSettingsCustomization::OnVerifyMappingName(const FTe
 	// Validate that the new name can be used
 	FEnhancedInputPlayerMappableNameValidator NameValidator(CurrentName);
 	EValidatorResult Result = NameValidator.IsValid(InNewText.ToString(), false);
-	OutErrorMessage = INameValidatorInterface::GetErrorText(InNewText.ToString(), Result);
+	OutErrorMessage = FEnhancedInputPlayerMappableNameValidator::GetErrorText(InNewText.ToString(), Result);
 
 	return Result == EValidatorResult::Ok || Result == EValidatorResult::ExistingName;
 }
diff --git a/Engine/Plugins/EnhancedInput/Source/InputEditor/Private/InputEditorModule.cpp b/Engine/Plugins/EnhancedInput/Source/InputEditor/Private/InputEditorModule.cpp
--- a/Engine/Plugins/EnhancedInput/Source/InputEditor/Private/InputEditorModule.cpp
+++ b/Engine/Plugins/EnhancedInput/Source/InputEditor/Private/InputEditorModule.cpp
@@ -555,19 +555,25 @@ void FInputEditorModule::AutoUpgradeDefaultInputClasses()
 	}
 }
 
-bool FInputEditorModule::IsMappingNameInUse(const FName InName)
+const UPlayerMappableKeySettings* FInputEditorModule::FindMappingByName(const FName InName)
 {
 	// Instanced objects that have been deleted will have an outer of the transient package, and we don't want to include them
 	for (TObjectIterator<UPlayerMappableKeySettings> Itr(RF_ClassDefaultObject | RF_Transient, true, EInternalObjectFlags::Garbage | EInternalObjectFlags::Unreachable); Itr; ++Itr)
 	{
-		UPlayerMappableKeySettings* Settings = *Itr;		
+		const UPlayerMappableKeySettings* Settings = *Itr;
 		if (Settings && Settings->GetOuter() != GetTransientPackage() && Settings->Name == InName)
 		{
-			return true;
+			// What if we had just removed a mapping of this name?
+			return Settings;
 		}
 	}
-	
-	return false;
+	return nullptr;
+}
+
+bool FInputEditorModule::IsMappingNameInUse(const FName InName)
+{
+	// If a mapping with this name exists, then it is in use
+	return FindMappingByName(InName) != nullptr;
 }
 
 void FInputEditorModule::Tick(float DeltaTime)
diff --git a/Engine/Plugins/EnhancedInput/Source/InputEditor/Public/EnhancedInputPlayerMappableNameValidator.h b/Engine/Plugins/EnhancedInput/Source/InputEditor/Public/EnhancedInputPlayerMappableNameValidator.h
--- a/Engine/Plugins/EnhancedInput/Source/InputEditor/Public/EnhancedInputPlayerMappableNameValidator.h
+++ b/Engine/Plugins/EnhancedInput/Source/InputEditor/Public/EnhancedInputPlayerMappableNameValidator.h
@@ -12,4 +12,12 @@ public:
 	// Begin FNameValidatorInterface
 	virtual EValidatorResult IsValid(const FString& Name, bool bOriginal) override;
 	// End FNameValidatorInterface
+
+	/**
+	 * If the error code is EValidatorResult::AlreadyInUse, this
+	 * returns special error text that says what asset is using the name.
+	 * 
+	 * Otherwise it will return INameValidatorInterface::GetErrorText.
+	 */
+	static FText GetErrorText(const FString& Name, EValidatorResult ErrorCode);
 };
\ No newline at end of file
diff --git a/Engine/Plugins/EnhancedInput/Source/InputEditor/Public/InputEditorModule.h b/Engine/Plugins/EnhancedInput/Source/InputEditor/Public/InputEditorModule.h
--- a/Engine/Plugins/EnhancedInput/Source/InputEditor/Public/InputEditorModule.h
+++ b/Engine/Plugins/EnhancedInput/Source/InputEditor/Public/InputEditorModule.h
@@ -16,6 +16,7 @@ DECLARE_LOG_CATEGORY_EXTERN(LogEnhancedInputEditor, Log, All);
 
 class UInputAction;
 class SWindow;
+class UPlayerMappableKeySettings;
 
 class FInputEditorModule : public IModuleInterface, public FTickableEditorObject
 {
@@ -33,6 +34,9 @@ public:
 
 	static EAssetTypeCategories::Type GetInputAssetsCategory() { return InputAssetsCategory; }
 	
+	/** Returns a pointer to the player mappable settings object that has this mapping name */
+	INPUTEDITOR_API static const UPlayerMappableKeySettings* FindMappingByName(const FName InName);
+
 	/** Returns true if the given name is in use by a player mappable key setting */
 	INPUTEDITOR_API static bool IsMappingNameInUse(const FName InName);
 		
