Expose LogEnhancedInput as part of the API so that it can be used in other modules

#rnx

diff --git a/Engine/Plugins/EnhancedInput/Source/EnhancedInput/Public/EnhancedInputModule.h b/Engine/Plugins/EnhancedInput/Source/EnhancedInput/Public/EnhancedInputModule.h
--- a/Engine/Plugins/EnhancedInput/Source/EnhancedInput/Public/EnhancedInputModule.h
+++ b/Engine/Plugins/EnhancedInput/Source/EnhancedInput/Public/EnhancedInputModule.h
@@ -10,7 +10,7 @@ struct FKey;
 
 class APlayerController;
 
-DECLARE_LOG_CATEGORY_EXTERN(LogEnhancedInput, Log, All);
+ENHANCEDINPUT_API DECLARE_LOG_CATEGORY_EXTERN(LogEnhancedInput, Log, All);
 
 struct ENHANCEDINPUT_API FEnhancedInputKeys
 {
