Add some helper functions to the user settings to access the key profile

#rnx

diff --git a/Engine/Plugins/EnhancedInput/Source/EnhancedInput/Private/UserSettings/EnhancedInputUserSettings.cpp b/Engine/Plugins/EnhancedInput/Source/EnhancedInput/Private/UserSettings/EnhancedInputUserSettings.cpp
--- a/Engine/Plugins/EnhancedInput/Source/EnhancedInput/Private/UserSettings/EnhancedInputUserSettings.cpp
+++ b/Engine/Plugins/EnhancedInput/Source/EnhancedInput/Private/UserSettings/EnhancedInputUserSettings.cpp
@@ -995,6 +995,19 @@ void UEnhancedInputUserSettings::ResetAllPlayerKeysInRow(const FMapPlayerKeyArgs
 	}
 }
 
+void UEnhancedInputUserSettings::ResetKeyProfileToDefault(const FGameplayTag& ProfileId, FGameplayTagContainer& FailureReason)
+{
+	if (UEnhancedPlayerMappableKeyProfile* KeyProfile = GetKeyProfileWithIdentifier(ProfileId))
+	{
+		KeyProfile->ResetToDefault();
+		OnSettingsChanged.Broadcast(this);
+	}
+	else
+	{
+		FailureReason.AddTag(UE::EnhancedInput::TAG_NoKeyProfile);
+	}
+}
+
 void UEnhancedInputUserSettings::OnKeyMappingUpdated(FPlayerKeyMapping* ChangedMapping, const FMapPlayerKeyArgs& InArgs, const bool bIsBeingUnmapped)
 {
 	// Do nothing by default
diff --git a/Engine/Plugins/EnhancedInput/Source/EnhancedInput/Public/UserSettings/EnhancedInputUserSettings.h b/Engine/Plugins/EnhancedInput/Source/EnhancedInput/Public/UserSettings/EnhancedInputUserSettings.h
--- a/Engine/Plugins/EnhancedInput/Source/EnhancedInput/Public/UserSettings/EnhancedInputUserSettings.h
+++ b/Engine/Plugins/EnhancedInput/Source/EnhancedInput/Public/UserSettings/EnhancedInputUserSettings.h
@@ -553,6 +553,15 @@ public:
 	UFUNCTION(BlueprintCallable, Category="Enhanced Input|User Settings", meta = (AutoCreateRefTerm = "FailureReason"))
 	virtual void ResetAllPlayerKeysInRow(const FMapPlayerKeyArgs& InArgs, FGameplayTagContainer& FailureReason);
 
+	/**
+	 * Resets the given key profile to default key mappings
+	 *
+	 * @param ProfileId		The ID of the key profile to reset
+	 * @param FailureReason	Populated with failure reasons if the operation fails.
+	 */
+	UFUNCTION(BlueprintCallable, Category="Enhanced Input|User Settings", meta = (AutoCreateRefTerm = "FailureReason"))
+	virtual void ResetKeyProfileToDefault(const FGameplayTag& ProfileId, FGameplayTagContainer& FailureReason);
+
 	/** Returns a set of all player key mappings for the given mapping name. */
 	UFUNCTION(BlueprintCallable, Category="Enhanced Input|User Settings")
 	virtual const TSet<FPlayerKeyMapping>& FindMappingsInRow(const FName MappingName) const;
