Add a virtual OnKeyMappingRegistered function that will get called for each key mapping once it is registered with the Enhanced Input user settings. This is a good place to do any Post Initalize logic per key mapping. For example, you are upgrading to this system and already had a key mapping system in place, you can set the key here!

#rnx

diff --git a/Engine/Plugins/EnhancedInput/Source/EnhancedInput/Private/UserSettings/EnhancedInputUserSettings.cpp b/Engine/Plugins/EnhancedInput/Source/EnhancedInput/Private/UserSettings/EnhancedInputUserSettings.cpp
--- a/Engine/Plugins/EnhancedInput/Source/EnhancedInput/Private/UserSettings/EnhancedInputUserSettings.cpp
+++ b/Engine/Plugins/EnhancedInput/Source/EnhancedInput/Private/UserSettings/EnhancedInputUserSettings.cpp
@@ -1287,6 +1287,8 @@ bool UEnhancedInputUserSettings::RegisterKeyMappingsToProfile(UEnhancedPlayerMap
 			// We always want to ensure that the metadata on a player key mapping is up to date
 			// from the Input Mapping Context.
 			ExistingMapping.UpdateMetadataFromActionKeyMapping(KeyMapping);
+
+			OnKeyMappingRegistered(ExistingMapping, KeyMapping);
 		}
 
 		// If the mapping was not found in the existing row, then the player has not mapped anything to it.
@@ -1294,7 +1296,8 @@ bool UEnhancedInputUserSettings::RegisterKeyMappingsToProfile(UEnhancedPlayerMap
 		if (!bUpdatedExistingMapping)
 		{
 			// Add a default mapping to this row
-			MappingRow.Mappings.Add({ KeyMapping, MappingSlot, MappingDeviceType });
+			const FSetElementId ElemId = MappingRow.Mappings.Add({ KeyMapping, MappingSlot, MappingDeviceType });
+			OnKeyMappingRegistered(MappingRow.Mappings.Get(ElemId), KeyMapping);
 		}
 
 		// Increment the mapping slot to keep track of how many slots are in use by which hardware device
@@ -1305,6 +1308,11 @@ bool UEnhancedInputUserSettings::RegisterKeyMappingsToProfile(UEnhancedPlayerMap
 	return true;
 }
 
+void UEnhancedInputUserSettings::OnKeyMappingRegistered(FPlayerKeyMapping& RegisteredMapping, const FEnhancedActionKeyMapping& SourceMapping)
+{
+	// Does nothing by default
+}
+
 FHardwareDeviceIdentifier UEnhancedInputUserSettings::DetermineHardwareDeviceForActionMapping(const FEnhancedActionKeyMapping& ActionMapping) const
 {
 	return FHardwareDeviceIdentifier::Invalid;
diff --git a/Engine/Plugins/EnhancedInput/Source/EnhancedInput/Public/UserSettings/EnhancedInputUserSettings.h b/Engine/Plugins/EnhancedInput/Source/EnhancedInput/Public/UserSettings/EnhancedInputUserSettings.h
--- a/Engine/Plugins/EnhancedInput/Source/EnhancedInput/Public/UserSettings/EnhancedInputUserSettings.h
+++ b/Engine/Plugins/EnhancedInput/Source/EnhancedInput/Public/UserSettings/EnhancedInputUserSettings.h
@@ -515,6 +515,21 @@ protected:
 	 */
 	virtual void OnKeyMappingUpdated(FPlayerKeyMapping* ChangedMapping, const FMapPlayerKeyArgs& InArgs, const bool bIsBeingUnmapped);
 
+
+	/**
+	 * A virtual function that will be called whenever a player key mapping is first registered. This will be called once
+	 * by UEnhancedInputUserSettings::RegisterKeyMappingsToProfile for each mapping that is created from the Input Mapping Context.
+	 *
+	 * This is a good place to do any "Post Registration" setup you may want to for you key mappings such as changing the
+	 * key that is mapped based on some other settings or data outside of this system (i.e. you are upgrading to this system
+	 * and already had a key mapping system in place, you can set the key here!)
+	 *
+	 * This does nothing by default.
+	 * @param RegisteredMapping		The newly registered player key mapping
+	 * @param SourceMapping			The source mapping that the registered mapping was created from. 
+	 */
+	virtual void OnKeyMappingRegistered(FPlayerKeyMapping& RegisteredMapping, const FEnhancedActionKeyMapping& SourceMapping);
+
 public:
 	
 	ULocalPlayer* GetLocalPlayer() const;
