Modify how the OnKeyMappingRegistered function is called so that you can potentially modify the MappingRow from within it. Do this by keeping track of what mapping elements were updated and register them after the initial iteration of existing mappings.

#rnx

diff --git a/Engine/Plugins/EnhancedInput/Source/EnhancedInput/Private/UserSettings/EnhancedInputUserSettings.cpp b/Engine/Plugins/EnhancedInput/Source/EnhancedInput/Private/UserSettings/EnhancedInputUserSettings.cpp
--- a/Engine/Plugins/EnhancedInput/Source/EnhancedInput/Private/UserSettings/EnhancedInputUserSettings.cpp
+++ b/Engine/Plugins/EnhancedInput/Source/EnhancedInput/Private/UserSettings/EnhancedInputUserSettings.cpp
@@ -1265,11 +1265,15 @@ bool UEnhancedInputUserSettings::RegisterKeyMappingsToProfile(UEnhancedPlayerMap
 		EPlayerMappableKeySlot& MappingSlot = MappingNameToSlotCount.FindOrAdd(MappingName);
 		
 		bool bUpdatedExistingMapping = false;
+
+		static TArray<FSetElementId> MappingsRegistered;
+		MappingsRegistered.Reset();
 		
 		// Iterate any existing mappings in this row to ensure that the metadata and default key values are correct.
 		// At this stage, keys will only exist here if they are player customized
-		for (FPlayerKeyMapping& ExistingMapping : MappingRow.Mappings)
-		{
+		for (auto It = MappingRow.Mappings.CreateIterator(); It; ++It)
+		{			
+			FPlayerKeyMapping& ExistingMapping = *It;
 			// We only want to update the default _key_ for an existing mapping if it is from
 			// the same slot as this Action Mapping would be going in.
 			// This is because players can map keys to slots are not defined in the input mapping context,
@@ -1287,8 +1291,9 @@ bool UEnhancedInputUserSettings::RegisterKeyMappingsToProfile(UEnhancedPlayerMap
 			// We always want to ensure that the metadata on a player key mapping is up to date
 			// from the Input Mapping Context.
 			ExistingMapping.UpdateMetadataFromActionKeyMapping(KeyMapping);
-
-			OnKeyMappingRegistered(ExistingMapping, KeyMapping);
+			
+			// Keep track of this mapping ID in the set so that we can register later
+			MappingsRegistered.Emplace(It.GetId());
 		}
 
 		// If the mapping was not found in the existing row, then the player has not mapped anything to it.
@@ -1296,8 +1301,18 @@ bool UEnhancedInputUserSettings::RegisterKeyMappingsToProfile(UEnhancedPlayerMap
 		if (!bUpdatedExistingMapping)
 		{
 			// Add a default mapping to this row
-			const FSetElementId ElemId = MappingRow.Mappings.Add({ KeyMapping, MappingSlot, MappingDeviceType });
-			OnKeyMappingRegistered(MappingRow.Mappings.Get(ElemId), KeyMapping);
+			FSetElementId ElemId = MappingRow.Mappings.Add({ KeyMapping, MappingSlot, MappingDeviceType });
+
+			// Keep track of this mapping ID in the set so that we can register later
+			MappingsRegistered.Emplace(ElemId);
+		}
+
+		for (const FSetElementId& RegisteredMappingId : MappingsRegistered)
+		{
+			if (RegisteredMappingId.IsValidId())
+			{
+				OnKeyMappingRegistered(MappingRow.Mappings.Get(RegisteredMappingId), KeyMapping);
+			}			
 		}
 
 		// Increment the mapping slot to keep track of how many slots are in use by which hardware device
