Correctly update the Display Category and AssociatedInputAction fields when creating a new player-made mapping.

#rnx

diff --git a/Engine/Plugins/EnhancedInput/Source/EnhancedInput/Private/UserSettings/EnhancedInputUserSettings.cpp b/Engine/Plugins/EnhancedInput/Source/EnhancedInput/Private/UserSettings/EnhancedInputUserSettings.cpp
--- a/Engine/Plugins/EnhancedInput/Source/EnhancedInput/Private/UserSettings/EnhancedInputUserSettings.cpp
+++ b/Engine/Plugins/EnhancedInput/Source/EnhancedInput/Private/UserSettings/EnhancedInputUserSettings.cpp
@@ -905,6 +905,8 @@ void UEnhancedInputUserSettings::MapPlayerKey(const FMapPlayerKeyArgs& InArgs, F
 			// EKeys::Invalid and we only need to track the player mapped key
 			PlayerMappingData.SetCurrentKey(InArgs.NewKey);
 			PlayerMappingData.DisplayName = ExistingMapping->DisplayName;
+			PlayerMappingData.DisplayCategory = ExistingMapping->DisplayCategory;
+			PlayerMappingData.AssociatedInputAction = ExistingMapping->AssociatedInputAction;
 			
 			const FSetElementId SetElem = MappingRow->Mappings.Add(PlayerMappingData);
 			OnKeyMappingUpdated(&MappingRow->Mappings.Get(SetElem), InArgs, false);
