Deprecate the EI Blueprint Library functions to access player mappable slots

#rnx

diff --git a/Engine/Plugins/EnhancedInput/Source/EnhancedInput/Public/EnhancedInputLibrary.h b/Engine/Plugins/EnhancedInput/Source/EnhancedInput/Public/EnhancedInputLibrary.h
--- a/Engine/Plugins/EnhancedInput/Source/EnhancedInput/Public/EnhancedInputLibrary.h
+++ b/Engine/Plugins/EnhancedInput/Source/EnhancedInput/Public/EnhancedInputLibrary.h
@@ -99,17 +99,25 @@ public:
 	UFUNCTION(BlueprintPure, meta = (DisplayName = "To String (InputActionValue)", CompactNodeTitle = "->", BlueprintAutocast), Category = "Utilities|String")
 	static FString Conv_InputActionValueToString(FInputActionValue ActionValue);
 
-	UFUNCTION(BlueprintPure, Category = "Input", meta = (ReturnDisplayName = "First Player Mappable Key Slot"))
+	PRAGMA_DISABLE_DEPRECATION_WARNINGS
+
+	UE_DEPRECATED(5.3, "FPlayerMappableKeySlot has been deprecated. Please use EPlayerMappableKeySlot instead.")
+	UFUNCTION(BlueprintPure, Category = "Input", meta = (ReturnDisplayName = "First Player Mappable Key Slot", DeprecatedFunction, DeprecationMessage="FPlayerMappableKeyOptions has been deprecated. Please use UPlayerMappableKeySettings instead."))
 	static FPlayerMappableKeySlot& GetFirstPlayerMappableKeySlot() { return FPlayerMappableKeySlot::FirstKeySlot; };
 
-	UFUNCTION(BlueprintPure, Category = "Input", meta = (ReturnDisplayName = "Second Player Mappable Key Slot"))
+	UE_DEPRECATED(5.3, "FPlayerMappableKeySlot has been deprecated. Please use EPlayerMappableKeySlot instead.")
+	UFUNCTION(BlueprintPure, Category = "Input", meta = (ReturnDisplayName = "Second Player Mappable Key Slot", DeprecatedFunction, DeprecationMessage="FPlayerMappableKeyOptions has been deprecated. Please use UPlayerMappableKeySettings instead."))
 	static FPlayerMappableKeySlot& GetSecondPlayerMappableKeySlot() { return FPlayerMappableKeySlot::SecondKeySlot; };
 
-	UFUNCTION(BlueprintPure, Category = "Input", meta = (ReturnDisplayName = "Third Player Mappable Key Slot"))
+	UE_DEPRECATED(5.3, "FPlayerMappableKeySlot has been deprecated. Please use EPlayerMappableKeySlot instead.")
+	UFUNCTION(BlueprintPure, Category = "Input", meta = (ReturnDisplayName = "Third Player Mappable Key Slot", DeprecatedFunction, DeprecationMessage="FPlayerMappableKeyOptions has been deprecated. Please use UPlayerMappableKeySettings instead."))
 	static FPlayerMappableKeySlot& GetThirdPlayerMappableKeySlot() { return FPlayerMappableKeySlot::ThirdKeySlot; };
 
-	UFUNCTION(BlueprintPure, Category = "Input", meta = (ReturnDisplayName = "Fourth Player Mappable Key Slot"))
+	UE_DEPRECATED(5.3, "FPlayerMappableKeySlot has been deprecated. Please use EPlayerMappableKeySlot instead.")
+	UFUNCTION(BlueprintPure, Category = "Input", meta = (ReturnDisplayName = "Fourth Player Mappable Key Slot", DeprecatedFunction, DeprecationMessage="FPlayerMappableKeyOptions has been deprecated. Please use UPlayerMappableKeySettings instead."))
 	static FPlayerMappableKeySlot& GetFourthPlayerMappableKeySlot() { return FPlayerMappableKeySlot::FourthKeySlot; };
+	
+	PRAGMA_ENABLE_DEPRECATION_WARNINGS
 };
 
 #if UE_ENABLE_INCLUDE_ORDER_DEPRECATED_IN_5_2
