Deprecate the Input Editor pieces that reference Player Mappable Input Configs

#prerflight 644ad724b208f61af83532af

diff --git a/Engine/Plugins/EnhancedInput/Source/InputEditor/Private/InputEditorModule.cpp b/Engine/Plugins/EnhancedInput/Source/InputEditor/Private/InputEditorModule.cpp
--- a/Engine/Plugins/EnhancedInput/Source/InputEditor/Private/InputEditorModule.cpp
+++ b/Engine/Plugins/EnhancedInput/Source/InputEditor/Private/InputEditorModule.cpp
@@ -219,6 +219,7 @@ UObject* UInputAction_Factory::FactoryCreateNew(UClass* Class, UObject* InParent
 	}
 }
 
+PRAGMA_DISABLE_DEPRECATION_WARNINGS
 // UPlayerMappableInputConfig_Factory
 UPlayerMappableInputConfig_Factory::UPlayerMappableInputConfig_Factory(const class FObjectInitializer& OBJ)
 	: Super(OBJ)
@@ -233,6 +234,7 @@ UObject* UPlayerMappableInputConfig_Factory::FactoryCreateNew(UClass* Class, UOb
 	check(Class->IsChildOf(UPlayerMappableInputConfig::StaticClass()));
 	return NewObject<UPlayerMappableInputConfig>(InParent, Class, Name, Flags | RF_Transactional, Context);
 }
+PRAGMA_ENABLE_DEPRECATION_WARNINGS
 
 // Asset type actions
 
@@ -313,7 +315,8 @@ void FAssetTypeActions_InputAction::GetCreateContextFromActionsMenu(TArray<TWeak
 	}
 }
 
-class FAssetTypeActions_PlayerMappableInputConfig : public FAssetTypeActions_DataAsset
+PRAGMA_DISABLE_DEPRECATION_WARNINGS
+class UE_DEPRECATED(5.3, "PlayerMappableInputConfig has been deprecated, please use UEnhancedInputUserSettings instead") FAssetTypeActions_PlayerMappableInputConfig : public FAssetTypeActions_DataAsset
 {
 public:
 	virtual FText GetName() const override { return NSLOCTEXT("AssetTypeActions", "AssetTypeActions_PlayerMappableInputConfig", "Player Mappable Input Config"); }
@@ -322,6 +325,7 @@ public:
 	virtual FText GetAssetDescription(const FAssetData& AssetData) const override { return NSLOCTEXT("AssetTypeActions", "AssetTypeActions_PlayerBindableInputConfigDesc", "Represents one set of Player Mappable controller/keymappings"); }
 	virtual UClass* GetSupportedClass() const override { return UPlayerMappableInputConfig::StaticClass(); }
 };
+PRAGMA_ENABLE_DEPRECATION_WARNINGS
 
 struct FInputActionGraphActions : public FAssetBlueprintGraphActions
 {
@@ -419,7 +423,9 @@ void FInputEditorModule::StartupModule()
 	{
 		RegisterAssetTypeActions(AssetTools, MakeShareable(new FAssetTypeActions_InputAction));
 		RegisterAssetTypeActions(AssetTools, MakeShareable(new FAssetTypeActions_InputContext));
+		PRAGMA_DISABLE_DEPRECATION_WARNINGS
 		RegisterAssetTypeActions(AssetTools, MakeShareable(new FAssetTypeActions_PlayerMappableInputConfig));
+		PRAGMA_ENABLE_DEPRECATION_WARNINGS
 		// TODO: Build these off a button on the InputContext Trigger/Mapping pickers? Would be good to have both.
 		//RegisterAssetTypeActions(AssetTools, MakeShareable(new FAssetTypeActions_InputTrigger));
 		//RegisterAssetTypeActions(AssetTools, MakeShareable(new FAssetTypeActions_InputModifier));
diff --git a/Engine/Plugins/EnhancedInput/Source/InputEditor/Public/InputEditorModule.h b/Engine/Plugins/EnhancedInput/Source/InputEditor/Public/InputEditorModule.h
--- a/Engine/Plugins/EnhancedInput/Source/InputEditor/Public/InputEditorModule.h
+++ b/Engine/Plugins/EnhancedInput/Source/InputEditor/Public/InputEditorModule.h
@@ -96,7 +96,7 @@ public:
 };
 
 UCLASS()
-class INPUTEDITOR_API UPlayerMappableInputConfig_Factory : public UFactory
+class UE_DEPRECATED(5.3, "UPlayerMappableInputConfig has been deprecated, please use the UEnhancedInputUserSettings system instead.") INPUTEDITOR_API UPlayerMappableInputConfig_Factory : public UFactory
 {
 	GENERATED_UCLASS_BODY()
 public:
