- Fixed drag n drop issue with rows having expanded items.
=> code from CL 17423372 was disabling it explicitely, but I couldn't find any reason or have any repro. The comment talk about disabling between an item and its child, while with a EItemDropZone::BelowItem, it correspond to the bottom part of an item, and not the bottom item, so having the child item expanded or not won't change the drag n drop behavior.
- Fixed drag n drop issue for Input Mapping Context, where drag n dropping could give deterministic but incorrect result when more than 1 mapping exists for a group.
=> the reordering algorithm used requires that the content of the array is ordered by group, while it was not necessarily the case. Starting from this CL, it will (slower but safer algorithm).
- Prevent the drag n drop for Input Mapping Context between rows that are not part of the same parent node.

Tests:
-opened various Input Mapping Context from QA Game, like /Maps/Framework/EnhancedInput/EI_HelperAssets/Controls_3rdPerson/IMC_3rdPerson_Controls.uasset, that was failing 100% on some actions. Drag n dropped groups (collapsed or expanded), even on the last one, and on individual entries with 100% success. Did a lot of drag n drop test (between close elements, first and last, last and first, first and middle,...)
-tried to drag n drop across entries from different parent => it was blocked as expected.
-undo/redo after those operations.
-used drag n drop in other contexts (media play list, data table), no regression noticed when drag n dropping elements.

diff --git a/Engine/Plugins/EnhancedInput/Source/InputEditor/Private/ActionMappingDetails.cpp b/Engine/Plugins/EnhancedInput/Source/InputEditor/Private/ActionMappingDetails.cpp
--- a/Engine/Plugins/EnhancedInput/Source/InputEditor/Private/ActionMappingDetails.cpp
+++ b/Engine/Plugins/EnhancedInput/Source/InputEditor/Private/ActionMappingDetails.cpp
@@ -57,7 +57,18 @@ int32 FActionMappingsNodeBuilderEx::GetNumGroupedMappings() const
 	return GroupedMappings.Num();	
 }
 
-void FActionMappingsNodeBuilderEx::ReorderMappings(int32 OriginalIndex, int32 NewIndex)
+int32 FActionMappingsNodeBuilderEx::GetNumMappings() const
+{
+	int32 NumMappings = 0;
+	for (const ActionMappingDetails::FMappingSet& MappingSet : GroupedMappings)
+	{
+		NumMappings += MappingSet.Mappings.Num();
+	}
+
+	return NumMappings;
+}
+
+void FActionMappingsNodeBuilderEx::ReorderMappings(int32 OriginalIndex, int32 NewIndex, EReorderMode ReorderMode)
 {
 	TArray<UObject*> OuterObjects;
 	ActionMappingsPropertyHandle->GetOuterObjects(OuterObjects);
@@ -77,59 +88,55 @@ void FActionMappingsNodeBuilderEx::ReorderMappings(int32 OriginalIndex, int32 Ne
 
 	DelayedGroupExpansionStates.Emplace(nullptr, true);
 	
-	const UInputAction* ActionOfNewIndex = GroupedMappings[NewIndex].SharedAction;
-	const UInputAction* ActionOfOriginalIndex = GroupedMappings[OriginalIndex].SharedAction;
-	int32 NewMappingIndex = INDEX_NONE;
-	int32 OriginalMappingIndex = INDEX_NONE;
-
-	const TSharedPtr<IPropertyHandleArray> ActionMappingsArrayHandle = ActionMappingsPropertyHandle->AsArray();
-	uint32 NumMappings;
-	ActionMappingsArrayHandle->GetNumElements(NumMappings);
-	// loop through all mappings to find the first instance of the action we're moving to and from
-	for (uint32 Index = 0; Index < NumMappings; ++Index)
+	if (ReorderMode == EReorderMode::Group)
 	{
-		TSharedRef<IPropertyHandle> ActionMapping = ActionMappingsArrayHandle->GetElement(Index);
-		const UObject* Action;
-		const FPropertyAccess::Result Result = ActionMapping->GetChildHandle(GET_MEMBER_NAME_CHECKED(FEnhancedActionKeyMapping, Action))->GetValue(Action);
+		// The mapping array we get here can have different entries with the groups interleaved. The following algo will ensure to clean that, while applying the reordering, without changing the order inside a given category.
 
-		if (Result == FPropertyAccess::Success)
+		//1: Get the list of groups in their current order, and assign them an index, with room to insert the moved group between any groups.
+		TMap<const UObject*, uint32> Groups;
+		uint32 GroupIndex = 1;
+		const TSharedPtr<IPropertyHandleArray> ActionMappingsArrayHandle = ActionMappingsPropertyHandle->AsArray();
+		uint32 NumMappings;
+		ActionMappingsArrayHandle->GetNumElements(NumMappings);
+		for (uint32 Index = 0; Index < NumMappings; ++Index)
 		{
-			// Only care about the first time we find an action of that type
-			if (Action == ActionOfNewIndex && NewMappingIndex == INDEX_NONE)
-			{
-				NewMappingIndex = Index;
-			}
-			else if (Action == ActionOfOriginalIndex && OriginalMappingIndex == INDEX_NONE)
+			TSharedRef<IPropertyHandle> ActionMapping = ActionMappingsArrayHandle->GetElement(Index);
+			const UObject* Action;
+			const FPropertyAccess::Result Result = ActionMapping->GetChildHandle(GET_MEMBER_NAME_CHECKED(FEnhancedActionKeyMapping, Action))->GetValue(Action);
+
+			if (Result == FPropertyAccess::Success)
 			{
-				OriginalMappingIndex = Index;
+				if (!Groups.Contains(Action))
+				{
+					Groups.Add(Action, GroupIndex);
+					GroupIndex += 2;
+				}
 			}
 		}
-	}
 
-	// make sure we found indices for each action
-	if (NewMappingIndex == INDEX_NONE || OriginalIndex == INDEX_NONE)
-	{
-		return;
-	}
+		//2: Compute a new index to insert at the right place.
+		const UInputAction* ActionOfOriginalIndex = GroupedMappings[OriginalIndex].SharedAction;
+		Groups[ActionOfOriginalIndex] = (NewIndex * 2 + ((OriginalIndex < NewIndex) ? 2 : 0));
 
-	// figuring out if shifting items forward or backwards
-	if ((OriginalMappingIndex - NewMappingIndex) > 0)
+		//3: Stable sort, to sort group in the new required order, but keep the order inside each category.
+		InputContext->Mappings.StableSort([&Groups](const FEnhancedActionKeyMapping& A, const FEnhancedActionKeyMapping& B) { return Groups[A.Action] < Groups[B.Action]; });
+	}
+	else if (ReorderMode == EReorderMode::Single)
 	{
-		FEnhancedActionKeyMapping OriginalKeyMapping = InputContext->Mappings[OriginalMappingIndex];
-		for (int32 Index = OriginalMappingIndex; Index > NewMappingIndex; --Index)
+		if (FMath::Abs(NewIndex - OriginalIndex) == 1) //If both reordered elements are next to each other, we can avoid the cost of a remove/insert.
+		{
+			InputContext->Mappings.Swap(OriginalIndex, NewIndex);
+		}
+		else
 		{
-			InputContext->Mappings[Index] = InputContext->Mappings[Index - 1];
+			FEnhancedActionKeyMapping EnhancedActionKeyMapping(InputContext->Mappings[OriginalIndex]);
+			InputContext->Mappings.RemoveAt(OriginalIndex);
+			InputContext->Mappings.Insert(EnhancedActionKeyMapping, NewIndex);
 		}
-		InputContext->Mappings[NewMappingIndex] = OriginalKeyMapping;
 	}
 	else
 	{
-		FEnhancedActionKeyMapping OriginalKeyMapping = InputContext->Mappings[OriginalMappingIndex];
-		for (int32 Index = OriginalMappingIndex; Index < NewMappingIndex; ++Index)
-		{
-			InputContext->Mappings[Index] = InputContext->Mappings[Index + 1];
-		}
-		InputContext->Mappings[NewMappingIndex] = OriginalKeyMapping;
+		checkf(false, TEXT("Unsupported value for enum EReorderMode!"));
 	}
 
 	ActionMappingsPropertyHandle->NotifyPostChange(EPropertyChangeType::ArrayAdd);
@@ -141,10 +148,11 @@ class FInputMappingIndexDragDropOp : public FDecoratedDragDropOp
 public:
 	DRAG_DROP_OPERATOR_TYPE(FUserDefinedEnumIndexDragDropOp, FDecoratedDragDropOp);
 
-	FInputMappingIndexDragDropOp(ActionMappingDetails::FMappingSet* InTargetMappingSet, int32 InInputMappingIndex, FActionMappingsNodeBuilderEx* InActionMappingNodeBuilder)
+	FInputMappingIndexDragDropOp(ActionMappingDetails::FMappingSet* InTargetMappingSet, int32 InInputMappingIndex, FActionMappingsNodeBuilderEx* InActionMappingNodeBuilder, bool InIsGroup)
 		: TargetMappingSet(InTargetMappingSet)
 		, InputMappingIndex(InInputMappingIndex)
 		, ActionMappingNodeBuilder(InActionMappingNodeBuilder)
+		, IsGroup(InIsGroup)
 	{
 		if (InTargetMappingSet->SharedAction)
 		{
@@ -191,34 +199,55 @@ public:
 		return InputMappingIndex;
 	}
 
+	bool GetIsGroup() const
+	{
+		return IsGroup;
+	}
+
 private:
 	ActionMappingDetails::FMappingSet* TargetMappingSet;
 	int32 InputMappingIndex;
 	FActionMappingsNodeBuilderEx* ActionMappingNodeBuilder;
 	FText InputMappingDisplayText;
+	bool IsGroup;
 };
 
-/** Handler for customizing the drag-and-drop behavior for input mappings, allowing them to be reordered */
-class FInputMappingIndexDragDropHandler : public IDetailDragDropHandler
+/** Handler base class for customizing the drag-and-drop behavior for input mappings (single or in group), allowing them to be reordered */
+class FInputMappingIndexDragDropHandlerBase : public IDetailDragDropHandler
 {
 public:
-	FInputMappingIndexDragDropHandler(ActionMappingDetails::FMappingSet* InTargetMappingSet, int32 InTargetMappingIndex, FActionMappingsNodeBuilderEx* InActionMappingNodeBuilder)
+	FInputMappingIndexDragDropHandlerBase(ActionMappingDetails::FMappingSet* InTargetMappingSet, int32 InTargetMappingIndex, FActionMappingsNodeBuilderEx* InActionMappingNodeBuilder, bool InIsGroup)
 		: TargetInputMappingSet(InTargetMappingSet)
 		, InputMappingIndex(InTargetMappingIndex)
 		, ActionMappingNodeBuilder(InActionMappingNodeBuilder)
+		, IsGroup(InIsGroup)
 	{
 		check(TargetInputMappingSet);
 		check(ActionMappingNodeBuilder);
-		check(InputMappingIndex >= 0 && InputMappingIndex < ActionMappingNodeBuilder->GetNumGroupedMappings());
 	}
 
 	virtual TSharedPtr<FDragDropOperation> CreateDragDropOperation() const override
 	{
-		TSharedPtr<FInputMappingIndexDragDropOp> DragOp = MakeShared<FInputMappingIndexDragDropOp>(TargetInputMappingSet, InputMappingIndex, ActionMappingNodeBuilder);
+		TSharedPtr<FInputMappingIndexDragDropOp> DragOp = MakeShared<FInputMappingIndexDragDropOp>(TargetInputMappingSet, InputMappingIndex, ActionMappingNodeBuilder, IsGroup);
 		DragOp->Init();
 		return DragOp;
 	}
 
+	virtual bool AcceptDrop(const FDragDropEvent& DragDropEvent, EItemDropZone DropZone) const override
+	{
+		const TSharedPtr<FInputMappingIndexDragDropOp> DragOp = DragDropEvent.GetOperationAs<FInputMappingIndexDragDropOp>();
+		if (!DragOp.IsValid() || DropZone == EItemDropZone::OntoItem)
+		{
+			return false;
+		}
+
+		const int32 NewIndex = ComputeNewIndex(DragOp->GetInputMappingIndex(), InputMappingIndex, DropZone);
+		ActionMappingNodeBuilder->ReorderMappings(DragOp->GetInputMappingIndex(), NewIndex, IsGroup ? FActionMappingsNodeBuilderEx::EReorderMode::Group: FActionMappingsNodeBuilderEx::EReorderMode::Single);
+		return true;
+	}
+
+protected:
+
 	/** Compute new target index based on drop zone (above vs below) */
 	static int32 ComputeNewIndex(int32 OriginalIndex, int32 DropOntoIndex, EItemDropZone DropZone)
 	{
@@ -239,17 +268,59 @@ public:
 		return ensure(NewIndex >= 0) ? NewIndex : 0;
 	}
 
-	virtual bool AcceptDrop(const FDragDropEvent& DragDropEvent, EItemDropZone DropZone) const override
+	ActionMappingDetails::FMappingSet* TargetInputMappingSet;
+	int32 InputMappingIndex;
+	FActionMappingsNodeBuilderEx* ActionMappingNodeBuilder;
+	bool IsGroup;
+};
+
+/** Handler for customizing the drag-and-drop behavior for input grouped mappings, allowing them to be reordered */
+class FInputGroupedMappingIndexDragDropHandler : public FInputMappingIndexDragDropHandlerBase
+{
+public:
+	FInputGroupedMappingIndexDragDropHandler(ActionMappingDetails::FMappingSet* InTargetMappingSet, int32 InTargetMappingIndex, FActionMappingsNodeBuilderEx* InActionMappingNodeBuilder)
+		: FInputMappingIndexDragDropHandlerBase(InTargetMappingSet, InTargetMappingIndex, InActionMappingNodeBuilder, /*InIsGroup*/true)
+	{
+		check(InputMappingIndex >= 0 && InputMappingIndex < ActionMappingNodeBuilder->GetNumGroupedMappings());
+	}
+
+	virtual TOptional<EItemDropZone> CanAcceptDrop(const FDragDropEvent& DragDropEvent, EItemDropZone DropZone) const override
 	{
 		const TSharedPtr<FInputMappingIndexDragDropOp> DragOp = DragDropEvent.GetOperationAs<FInputMappingIndexDragDropOp>();
-		if (!DragOp.IsValid() || DropZone == EItemDropZone::OntoItem)
+		if (!DragOp.IsValid())
 		{
-			return false;
+			return TOptional<EItemDropZone>();
 		}
 
-		const int32 NewIndex = ComputeNewIndex(DragOp->GetInputMappingIndex(), InputMappingIndex, DropZone);
-		ActionMappingNodeBuilder->ReorderMappings(DragOp->GetInputMappingIndex(), NewIndex);
-		return true;
+		if (!DragOp->GetIsGroup())
+		{
+			return TOptional<EItemDropZone>();
+		}
+
+		// We're reordering, so there's no logical interpretation for dropping directly onto another item.
+		// Just change it to a drop-above in this case.
+		const EItemDropZone OverrideZone = (DropZone == EItemDropZone::BelowItem) ? EItemDropZone::BelowItem : EItemDropZone::AboveItem;
+		const int32 NewIndex = ComputeNewIndex(DragOp->GetInputMappingIndex(), InputMappingIndex, OverrideZone);
+
+		// Make sure that the new index is valid *and* that it represents an actual move from the current position.
+		if (NewIndex < 0 || NewIndex >= ActionMappingNodeBuilder->GetNumGroupedMappings() || NewIndex == DragOp->GetInputMappingIndex())
+		{
+			return TOptional<EItemDropZone>();
+		}
+
+		DragOp->SetValidTarget(true);
+		return OverrideZone;
+	}
+};
+
+/** Handler for customizing the drag-and-drop behavior for input mappings, allowing them to be reordered */
+class FInputMappingIndexDragDropHandler : public FInputMappingIndexDragDropHandlerBase
+{
+public:
+	FInputMappingIndexDragDropHandler(ActionMappingDetails::FMappingSet* InParentTargetMappingSet, int32 InTargetMappingIndex, FActionMappingsNodeBuilderEx* InActionMappingNodeBuilder)
+		: FInputMappingIndexDragDropHandlerBase(InParentTargetMappingSet, InTargetMappingIndex, InActionMappingNodeBuilder, /*InIsGroup*/false)
+	{
+		check(InputMappingIndex >= 0 && InputMappingIndex < ActionMappingNodeBuilder->GetNumMappings());
 	}
 
 	virtual TOptional<EItemDropZone> CanAcceptDrop(const FDragDropEvent& DragDropEvent, EItemDropZone DropZone) const override
@@ -260,13 +331,24 @@ public:
 			return TOptional<EItemDropZone>();
 		}
 
+		if (DragOp->GetIsGroup())
+		{
+			return TOptional<EItemDropZone>();
+		}
+
+		//Forbid drag n drop between input mapping owned different parents.
+		if (DragOp->GetMappingSet() != TargetInputMappingSet)
+		{
+			return TOptional<EItemDropZone>();
+		}
+
 		// We're reordering, so there's no logical interpretation for dropping directly onto another item.
 		// Just change it to a drop-above in this case.
 		const EItemDropZone OverrideZone = (DropZone == EItemDropZone::BelowItem) ? EItemDropZone::BelowItem : EItemDropZone::AboveItem;
 		const int32 NewIndex = ComputeNewIndex(DragOp->GetInputMappingIndex(), InputMappingIndex, OverrideZone);
 
 		// Make sure that the new index is valid *and* that it represents an actual move from the current position.
-		if (NewIndex < 0 || NewIndex >= ActionMappingNodeBuilder->GetNumGroupedMappings() || NewIndex == DragOp->GetInputMappingIndex())
+		if (NewIndex < 0 || NewIndex >= ActionMappingNodeBuilder->GetNumMappings() || NewIndex == DragOp->GetInputMappingIndex())
 		{
 			return TOptional<EItemDropZone>();
 		}
@@ -274,11 +356,6 @@ public:
 		DragOp->SetValidTarget(true);
 		return OverrideZone;
 	}
-
-private:
-	ActionMappingDetails::FMappingSet* TargetInputMappingSet;
-	int32 InputMappingIndex;
-	FActionMappingsNodeBuilderEx* ActionMappingNodeBuilder;
 };
 
 void FActionMappingsNodeBuilderEx::GenerateHeaderRowContent(FDetailWidgetRow& NodeRow)
@@ -369,6 +446,8 @@ void FActionMappingsNodeBuilderEx::GenerateChildContent(IDetailChildrenBuilder&
 {
 	RebuildGroupedMappings();
 
+	int32 GlobalIndex = 0;
+
 	for (int32 Index = 0; Index < GroupedMappings.Num(); ++Index)
 	{
 		ActionMappingDetails::FMappingSet& MappingSet = GroupedMappings[Index];
@@ -419,11 +498,15 @@ void FActionMappingsNodeBuilderEx::GenerateChildContent(IDetailChildrenBuilder&
 				RemoveButton
 			]
 		]
-		.DragDropHandler(MakeShared<FInputMappingIndexDragDropHandler>(&GroupedMappings[Index], Index, this));
+		.DragDropHandler(MakeShared<FInputGroupedMappingIndexDragDropHandler>(&GroupedMappings[Index], Index, this));
 
 		for (int32 MappingIndex = 0; MappingIndex < MappingSet.Mappings.Num(); ++MappingIndex)
 		{
-			ActionMappingGroup.AddPropertyRow(MappingSet.Mappings[MappingIndex]).ShowPropertyButtons(false);
+			ActionMappingGroup.AddPropertyRow(MappingSet.Mappings[MappingIndex])
+			.ShowPropertyButtons(false)
+			.DragDropHandler(MakeShared<FInputMappingIndexDragDropHandler>(&GroupedMappings[Index], GlobalIndex, this));
+
+			++GlobalIndex;
 		}
 	}
 }
diff --git a/Engine/Plugins/EnhancedInput/Source/InputEditor/Private/ActionMappingDetails.h b/Engine/Plugins/EnhancedInput/Source/InputEditor/Private/ActionMappingDetails.h
--- a/Engine/Plugins/EnhancedInput/Source/InputEditor/Private/ActionMappingDetails.h
+++ b/Engine/Plugins/EnhancedInput/Source/InputEditor/Private/ActionMappingDetails.h
@@ -46,7 +46,14 @@ public:
 	virtual FName GetName() const override { return FName(TEXT("ActionMappings")); }
 
 	int32 GetNumGroupedMappings() const;
-	void ReorderMappings(int32 OriginalIndex, int32 NewIndex);
+	int32 GetNumMappings() const;
+
+	enum class EReorderMode : uint8
+	{
+		Group,
+		Single
+	};
+	void ReorderMappings(int32 OriginalIndex, int32 NewIndex, EReorderMode ReorderMode);
 private:
 	void AddActionMappingButton_OnClick();
 	void ClearActionMappingButton_OnClick();
diff --git a/Engine/Source/Editor/PropertyEditor/Private/SDetailSingleItemRow.cpp b/Engine/Source/Editor/PropertyEditor/Private/SDetailSingleItemRow.cpp
--- a/Engine/Source/Editor/PropertyEditor/Private/SDetailSingleItemRow.cpp
+++ b/Engine/Source/Editor/PropertyEditor/Private/SDetailSingleItemRow.cpp
@@ -93,19 +93,6 @@ bool SDetailSingleItemRow::CheckValidDrop(const TSharedPtr<SDetailSingleItemRow>
 		return false;
 	}
 
-	// Disallow dropping immediately below an expanded item (e.g. array of structs) since the drop zone is between the item and its children
-	// User can drop above the next array item instead
-	// For example:
-	// - A    <- allow drop below A because it's collapsed
-	// + B    <- don't allow drop below B because it's expanded;
-	// | - a     dropping here seems to "parent" the dropped row to B
-	// | - b
-	// - C    <- user can drop above C instead
-	if (DropZone == EItemDropZone::BelowItem && IsItemExpanded())
-	{
-		return false;
-	}
-
 	TSharedPtr<FPropertyNode> SwappingPropertyNode = RowPtr->SwappablePropertyNode;
 	if (SwappingPropertyNode.IsValid() && SwappablePropertyNode.IsValid() && SwappingPropertyNode != SwappablePropertyNode)
 	{
@@ -229,12 +216,6 @@ FReply SDetailSingleItemRow::OnCustomAcceptDrop(const FDragDropEvent& DragDropEv
 
 TOptional<EItemDropZone> SDetailSingleItemRow::OnCustomCanAcceptDrop(const FDragDropEvent& DragDropEvent, EItemDropZone DropZone, TSharedPtr<FDetailTreeNode> Type)
 {
-	// Disallow drop between expanded parent item and its first child
-	if (DropZone == EItemDropZone::BelowItem && IsItemExpanded())
-	{
-		return TOptional<EItemDropZone>();
-	}
-
 	// This should only be registered as a delegate if there's a custom handler
 	if (ensure(WidgetRow.CustomDragDropHandler))
 	{
