Making input slots support enum ranges.

diff --git a/Engine/Plugins/EnhancedInput/Source/EnhancedInput/Public/UserSettings/EnhancedInputUserSettings.h b/Engine/Plugins/EnhancedInput/Source/EnhancedInput/Public/UserSettings/EnhancedInputUserSettings.h
--- a/Engine/Plugins/EnhancedInput/Source/EnhancedInput/Public/UserSettings/EnhancedInputUserSettings.h
+++ b/Engine/Plugins/EnhancedInput/Source/EnhancedInput/Public/UserSettings/EnhancedInputUserSettings.h
@@ -43,6 +43,8 @@ enum class EPlayerMappableKeySlot : uint8
 	Max
 };
 
+ENUM_RANGE_BY_FIRST_AND_LAST(EPlayerMappableKeySlot, EPlayerMappableKeySlot::First, EPlayerMappableKeySlot::Seventh);
+
 /** Arguments that can be used when mapping a player key */
 USTRUCT(BlueprintType)
 struct ENHANCEDINPUT_API FMapPlayerKeyArgs final
