Fixing missing include

diff --git a/Engine/Plugins/EnhancedInput/Source/EnhancedInput/Public/UserSettings/EnhancedInputUserSettings.h b/Engine/Plugins/EnhancedInput/Source/EnhancedInput/Public/UserSettings/EnhancedInputUserSettings.h
--- a/Engine/Plugins/EnhancedInput/Source/EnhancedInput/Public/UserSettings/EnhancedInputUserSettings.h
+++ b/Engine/Plugins/EnhancedInput/Source/EnhancedInput/Public/UserSettings/EnhancedInputUserSettings.h
@@ -2,20 +2,21 @@
 
 #pragma once
 
+#include "GameFramework/InputSettings.h"
 #include "GameFramework/SaveGame.h"
-#include "UObject/Object.h"
-#include "InputCoreTypes.h"					// For FKey
-#include "GameFramework/InputSettings.h"	// For FHardwareDeviceIdentifier
 #include "GameplayTagContainer.h"
+#include "InputCoreTypes.h"
+#include "Misc/EnumRange.h"
+#include "UObject/Object.h"
 
 #include "EnhancedInputUserSettings.generated.h"
 
-class UInputMappingContext;
+class UCanvas;
 class UEnhancedPlayerInput;
+class UInputAction;
+class UInputMappingContext;
 class ULocalPlayer;
 struct FEnhancedActionKeyMapping;
-class UInputAction;
-class UCanvas;
 
 /**
  * The "Slot" that a player mappable key is in.
