PR #10070: Enhanced Input: Resetting the value of the action value bindings when the action is not present in the current mappings (Contributed by [REDACTED_PERSON])

diff --git a/Engine/Plugins/EnhancedInput/Source/EnhancedInput/Private/EnhancedPlayerInput.cpp b/Engine/Plugins/EnhancedInput/Source/EnhancedInput/Private/EnhancedPlayerInput.cpp
--- a/Engine/Plugins/EnhancedInput/Source/EnhancedInput/Private/EnhancedPlayerInput.cpp
+++ b/Engine/Plugins/EnhancedInput/Source/EnhancedInput/Private/EnhancedPlayerInput.cpp
@@ -527,10 +527,19 @@ void UEnhancedPlayerInput::ProcessInputStack(const TArray<UInputComponent*>& Inp
 		// Update action value bindings
 		for (const FEnhancedInputActionValueBinding& Binding : IC->GetActionValueBindings())
 		{
-			// PERF: Lots of map lookups! Group EnhancedActionBindings by Action?
-			if (const FInputActionInstance* ActionData = FindActionInstanceData(Binding.GetAction()))
+			if (const UInputAction* Action = Binding.GetAction())
 			{
-				Binding.CurrentValue = ActionData->GetValue();
+				// PERF: Lots of map lookups! Group EnhancedActionBindings by Action?
+				if (const FInputActionInstance* ActionData = FindActionInstanceData(Action))
+				{
+					Binding.CurrentValue = ActionData->GetValue();
+				}
+				// If there is no action instance data related to this action, then reset the binding's value to zero
+				// to ensure that it gets its Completed state sent to any listeners
+				else
+				{
+					Binding.CurrentValue = FInputActionValue(Action->ValueType, FVector::ZeroVector);
+				}
 			}
 		}
 
