Limit the length of a player mappable name to the max length of FName's to properly check for the max length.

diff --git a/Engine/Plugins/EnhancedInput/Source/InputEditor/Private/EnhancedInputPlayerMappableNameValidator.cpp b/Engine/Plugins/EnhancedInput/Source/InputEditor/Private/EnhancedInputPlayerMappableNameValidator.cpp
--- a/Engine/Plugins/EnhancedInput/Source/InputEditor/Private/EnhancedInputPlayerMappableNameValidator.cpp
+++ b/Engine/Plugins/EnhancedInput/Source/InputEditor/Private/EnhancedInputPlayerMappableNameValidator.cpp
@@ -16,6 +16,15 @@ namespace UE::EnhancedInput
 	bEnableMappingNameValidation,
 	TEXT("Enables editor validation on player mapping names"),
 	ECVF_Default);
+	
+	/**
+	 * Returns the max acceptable length for a player mappable name.
+	 */
+	static int32 GetMaxAcceptableLength()
+	{
+		// Player Mappable names are just FNames, so using the max FName size here is acceptable.
+		return NAME_SIZE;
+	}
 }
 
 FEnhancedInputPlayerMappableNameValidator::FEnhancedInputPlayerMappableNameValidator(FName InExistingName)
@@ -24,6 +33,12 @@ FEnhancedInputPlayerMappableNameValidator::FEnhancedInputPlayerMappableNameValid
 
 EValidatorResult FEnhancedInputPlayerMappableNameValidator::IsValid(const FString& Name, bool bOriginal)
 {
+	// Ensure that the length of this name is a valid length
+	if (Name.Len() >= UE::EnhancedInput::GetMaxAcceptableLength())
+	{
+		return EValidatorResult::TooLong;
+	}
+	
 	EValidatorResult Result = FStringSetNameValidator::IsValid(Name, bOriginal);
 
 	if (UE::EnhancedInput::bEnableMappingNameValidation &&
@@ -49,6 +64,13 @@ FText FEnhancedInputPlayerMappableNameValidator::GetErrorText(const FString& Nam
 			return FText::Format(NSLOCTEXT("EnhancedInput", "MappingNameInUseBy_Error", "Name is already in use by '{AssetUsingName}'"), Args);	
 		}
 	}
+	else if (ErrorCode == EValidatorResult::TooLong)
+	{
+		FFormatNamedArguments Args;
+		Args.Add(TEXT("MaxLength"), UE::EnhancedInput::GetMaxAcceptableLength());
+
+		return FText::Format(NSLOCTEXT("EnhancedInput", "MappingNameTooLong_Error", "Names must be fewer then '{MaxLength}' characters"), Args);
+	}
 	
 	return INameValidatorInterface::GetErrorText(Name, ErrorCode);
 }
\ No newline at end of file
