PR #10403: Fix not being able to Copy & Paste the Mappings list in an Input Mapp? (Contributed by [REDACTED_PERSON])

diff --git a/Engine/Plugins/EnhancedInput/Source/InputEditor/Private/ActionMappingDetails.cpp b/Engine/Plugins/EnhancedInput/Source/InputEditor/Private/ActionMappingDetails.cpp
--- a/Engine/Plugins/EnhancedInput/Source/InputEditor/Private/ActionMappingDetails.cpp
+++ b/Engine/Plugins/EnhancedInput/Source/InputEditor/Private/ActionMappingDetails.cpp
@@ -372,7 +372,13 @@ void FActionMappingsNodeBuilderEx::GenerateHeaderRowContent(FDetailWidgetRow& No
 	ActionMappingsPropertyHandle->SetOnPropertyValueChanged(RebuildChildrenDelegate);
 	ActionMappingsPropertyHandle->AsArray()->SetOnNumElementsChanged(RebuildChildrenDelegate);
 
+	FUIAction CopyAction;
+	FUIAction PasteAction;
+	ActionMappingsPropertyHandle->CreateDefaultPropertyCopyPasteActions(CopyAction, PasteAction);
+
 	NodeRow
+	.CopyAction(CopyAction)
+	.PasteAction(PasteAction)
 	.FilterString(ActionMappingsPropertyHandle->GetPropertyDisplayName())
 	[
 		SNew(SHorizontalBox)
