Make the display name for getting the enhanced input user settings more explicit!

#rnx

diff --git a/Engine/Plugins/EnhancedInput/Source/EnhancedInput/Public/EnhancedInputSubsystemInterface.h b/Engine/Plugins/EnhancedInput/Source/EnhancedInput/Public/EnhancedInputSubsystemInterface.h
--- a/Engine/Plugins/EnhancedInput/Source/EnhancedInput/Public/EnhancedInputSubsystemInterface.h
+++ b/Engine/Plugins/EnhancedInput/Source/EnhancedInput/Public/EnhancedInputSubsystemInterface.h
@@ -73,7 +73,7 @@ public:
 
 	virtual UEnhancedPlayerInput* GetPlayerInput() const = 0;
 
-	UFUNCTION(BlueprintCallable, Category="Enhanced Input|User Settings")
+	UFUNCTION(BlueprintCallable, Category="Enhanced Input|User Settings", meta=(DisplayName="Get Enhanced Input User Settings"))
 	virtual UEnhancedInputUserSettings* GetUserSettings() const;
 
 protected:
