Add a Slate Notification if you add an Input Action that already has mappings in it within an IMC

#rnx

diff --git a/Engine/Plugins/EnhancedInput/Source/InputEditor/Private/ActionMappingDetails.cpp b/Engine/Plugins/EnhancedInput/Source/InputEditor/Private/ActionMappingDetails.cpp
--- a/Engine/Plugins/EnhancedInput/Source/InputEditor/Private/ActionMappingDetails.cpp
+++ b/Engine/Plugins/EnhancedInput/Source/InputEditor/Private/ActionMappingDetails.cpp
@@ -12,6 +12,8 @@
 #include "DragAndDrop/DecoratedDragDropOp.h"
 #include "AssetRegistry/AssetRegistryModule.h"
 #include "Widgets/Views/STableRow.h"
+#include "Framework/Notifications/NotificationManager.h"
+#include "Widgets/Notifications/SNotificationList.h"
 
 #define LOCTEXT_NAMESPACE "ActionMappingDetails"
 
@@ -558,6 +560,30 @@ void FActionMappingsNodeBuilderEx::OnActionMappingActionChanged(const FAssetData
 
 	if (SelectedAction != CurrentAction)
 	{
+		// If the IMC already has a mapping to this Input Action, then all of it's mappings will be moved.
+		// This can be jarring for the user if they are working with a complex IMC, so we will put a
+		// small toast up to tell them about it if we can
+		{
+			TArray<UObject*> OuterObjects;
+			ActionMappingsPropertyHandle->GetOuterObjects(OuterObjects);
+			if (!OuterObjects.IsEmpty())
+			{
+				if (const UInputMappingContext* InputContext = CastChecked<UInputMappingContext>(OuterObjects[0]))
+				{
+					if (InputContext->Mappings.FindByPredicate([&SelectedAction](const FEnhancedActionKeyMapping& Mapping) { return Mapping.Action == SelectedAction; }) != nullptr)
+					{
+						FFormatNamedArguments Args;
+						Args.Add(TEXT("ExistingAction"), FText::FromString(*GetNameSafe(SelectedAction)));
+						const FText NotifText = FText::Format(LOCTEXT("CombiningInputActionMappings_Notif", "'{ExistingAction}' already has mappings in this context, combining them!"), Args);
+						
+						FNotificationInfo Notif(NotifText);
+						Notif.ExpireDuration = 3.0f;
+						FSlateNotificationManager::Get().AddNotification(Notif);						
+					}
+				}	
+			}
+		}
+		
 		for (int32 Index = 0; Index < MappingSet.Mappings.Num(); ++Index)
 		{
 			MappingSet.Mappings[Index]->GetChildHandle(GET_MEMBER_NAME_CHECKED(FEnhancedActionKeyMapping, Action))->SetValue(SelectedAction);
