Allow blueprints to change the bConsumeInput flag on an Input Action at runtime

#rnx

diff --git a/Engine/Plugins/EnhancedInput/Source/EnhancedInput/Public/InputAction.h b/Engine/Plugins/EnhancedInput/Source/EnhancedInput/Public/InputAction.h
--- a/Engine/Plugins/EnhancedInput/Source/EnhancedInput/Public/InputAction.h
+++ b/Engine/Plugins/EnhancedInput/Source/EnhancedInput/Public/InputAction.h
@@ -52,7 +52,7 @@ public:
 	FText ActionDescription = FText::GetEmpty();
 
 	// Should this action swallow any inputs bound to it or allow them to pass through to affect lower priority bound actions?
-	UPROPERTY(EditAnywhere, BlueprintReadOnly, Category = Action)
+	UPROPERTY(EditAnywhere, BlueprintReadWrite, Category = Action)
 	bool bConsumeInput = true;
 
 	// Should this action be able to trigger whilst the game is paused - Replaces bExecuteWhenPaused
