Update tooltips/comments on Input Trigger states. Add a note about how to use Enhanced Input debug commands!

#rnx

diff --git a/Engine/Plugins/EnhancedInput/Source/EnhancedInput/Public/InputTriggers.h b/Engine/Plugins/EnhancedInput/Source/EnhancedInput/Public/InputTriggers.h
--- a/Engine/Plugins/EnhancedInput/Source/EnhancedInput/Public/InputTriggers.h
+++ b/Engine/Plugins/EnhancedInput/Source/EnhancedInput/Public/InputTriggers.h
@@ -33,22 +33,23 @@ enum class ETriggerEvent : uint8
 {
 	// No significant trigger state changes occurred and there are no active device inputs
 	None		= (0x0)		UMETA(Hidden),
-
 	// Triggering occurred after one or more processing ticks
 	Triggered	= (1 << 0),	// ETriggerState (None -> Triggered, Ongoing -> Triggered, Triggered -> Triggered)
 	
-	// An event has occurred that has begun Trigger evaluation. Note: Triggered may also occur this frame.
+	// An event has occurred that has begun Trigger evaluation. Note: Triggered may also occur this frame, but this event will always be fired first.
 	Started		= (1 << 1),	// ETriggerState (None -> Ongoing, None -> Triggered)
 
-	// Triggering is still being processed
+	// Triggering is still being processed. For example, an action with a "Press and Hold" trigger
+	// will be "Ongoing" while the user is holding down the key but the time threshold has not been met yet. 
 	Ongoing		= (1 << 2),	// ETriggerState (Ongoing -> Ongoing)
 
-	// Triggering has been canceled
+	// Triggering has been canceled. For example,  the user has let go of a key before the "Press and Hold" time threshold.
+	// The action has started to be evaluated, but never completed. 
 	Canceled	= (1 << 3),	// ETriggerState (Ongoing -> None)
 
 	// The trigger state has transitioned from Triggered to None this frame, i.e. Triggering has finished.
-	// NOTE: Using this event restricts you to one set of triggers for Started/Completed events. You may prefer two actions, each with its own trigger rules.
-	// TODO: Completed will not fire if any trigger reports Ongoing on the same frame, but both should fire. e.g. Tick 2 of Hold (= Ongoing) + Pressed (= None) combo will raise Ongoing event only.
+	// Note: Using this event restricts you to one set of triggers for Started/Completed events. You may prefer two actions, each with its own trigger rules.
+	// Completed will not fire if any trigger reports Ongoing on the same frame, but both should fire. e.g. Tick 2 of Hold (= Ongoing) + Pressed (= None) combo will raise Ongoing event only.
 	Completed	= (1 << 4),	// ETriggerState (Triggered -> None)
 };
 ENUM_CLASS_FLAGS(ETriggerEvent)
diff --git a/Engine/Plugins/EnhancedInput/Source/InputBlueprintNodes/Private/K2Node_EnhancedInputAction.cpp b/Engine/Plugins/EnhancedInput/Source/InputBlueprintNodes/Private/K2Node_EnhancedInputAction.cpp
--- a/Engine/Plugins/EnhancedInput/Source/InputBlueprintNodes/Private/K2Node_EnhancedInputAction.cpp
+++ b/Engine/Plugins/EnhancedInput/Source/InputBlueprintNodes/Private/K2Node_EnhancedInputAction.cpp
@@ -201,8 +201,9 @@ FText UK2Node_EnhancedInputAction::GetTooltipText() const
 		FString ActionPath = InputAction ? InputAction->GetFullName() : TEXT("");
 		CachedTooltip.SetCachedText(
 			FText::Format(
-				LOCTEXT("EnhancedInputAction_Tooltip", "Event for when '{0}' triggers.\n\nNOTE: This is not guaranteed to fire every frame, only when the Action is triggered and the current Input Mode includes 'Game'."),
-			FText::FromString(ActionPath)),
+				LOCTEXT("EnhancedInputAction_Tooltip", "Event for when '{0}' triggers.\n\nNote: This is not guaranteed to fire every frame, only when the Action is triggered and the current Input Mode includes 'Game'.\n\n{1}"),
+			FText::FromString(ActionPath),
+			LOCTEXT("EnhancedInputAction_Node_Tooltip_Tip", "Tip: Use the 'showdebug enhancedinput' command while playing to see debug information about Enhanced Input.")),
 			this);
 	}
 	return CachedTooltip;
