Remove dependency on InputBlueprintNodes from the InputEditor Module.

The Input Editor module has several features that are using the K2Nodes from enhanced input when it doesnt need to. The Blueprint Node logic should be kept in its own module.

This will allow the K2Nodes to actually read from editor settings and preferences in the rest of enhanced input, which would otherwise cause a circular depedency in the build graph.

This change is really just moving some functions from one module to another, neither module is publically exposed so no deprecation is required.

#rnx

diff --git a/Engine/Plugins/EnhancedInput/EnhancedInput.uplugin b/Engine/Plugins/EnhancedInput/EnhancedInput.uplugin
--- a/Engine/Plugins/EnhancedInput/EnhancedInput.uplugin
+++ b/Engine/Plugins/EnhancedInput/EnhancedInput.uplugin
@@ -23,7 +23,7 @@
 		{
 			"Name": "InputBlueprintNodes",
 			"Type": "UncookedOnly",
-			"LoadingPhase": "PreDefault"
+			"LoadingPhase": "Default"
 		},
 		{
 			"Name": "InputEditor",
diff --git a/Engine/Plugins/EnhancedInput/Source/InputBlueprintNodes/InputBlueprintNodes.Build.cs b/Engine/Plugins/EnhancedInput/Source/InputBlueprintNodes/InputBlueprintNodes.Build.cs
--- a/Engine/Plugins/EnhancedInput/Source/InputBlueprintNodes/InputBlueprintNodes.Build.cs
+++ b/Engine/Plugins/EnhancedInput/Source/InputBlueprintNodes/InputBlueprintNodes.Build.cs
@@ -6,6 +6,8 @@ namespace UnrealBuildTool.Rules
 	{
         public InputBlueprintNodes(ReadOnlyTargetRules Target) : base(Target)
         {
+	        PublicDependencyModuleNames.AddRange(new string[] { "InputEditor", });
+	        
             PrivateDependencyModuleNames.AddRange(
 				new string[] {
 					"BlueprintGraph",
diff --git a/Engine/Plugins/EnhancedInput/Source/InputBlueprintNodes/Private/InputBlueprintNodesModule.cpp b/Engine/Plugins/EnhancedInput/Source/InputBlueprintNodes/Private/InputBlueprintNodesModule.cpp
--- a/Engine/Plugins/EnhancedInput/Source/InputBlueprintNodes/Private/InputBlueprintNodesModule.cpp
+++ b/Engine/Plugins/EnhancedInput/Source/InputBlueprintNodes/Private/InputBlueprintNodesModule.cpp
@@ -1,13 +1,149 @@
 // Copyright Epic Games, Inc. All Rights Reserved.
 
 #include "InputBlueprintNodesModule.h"
+#include "K2Node_EnhancedInputAction.h"
+#include "K2Node_GetInputActionValue.h"
+#include "AssetBlueprintGraphActions.h"
+#include "InputAction.h"
+#include "Kismet2/KismetEditorUtilities.h"
+#include "EdGraphSchema_K2_Actions.h"
+#include "BlueprintEditorModule.h"
+#include "Framework/Notifications/NotificationManager.h"
+#include "Widgets/Notifications/SNotificationList.h"
+#include "BlueprintNodeTemplateCache.h"
+#include "TickableEditorObject.h"
+#include "UObject/UObjectIterator.h"
 
 #include "Modules/ModuleManager.h"
 
 #define LOCTEXT_NAMESPACE "InputBlueprintNodes"
 
-class FInputBlueprintNodesModule : public IModuleInterface
+struct FInputActionGraphActions : public FAssetBlueprintGraphActions
 {
+	virtual FText GetGraphHoverMessage(const FAssetData& AssetData, const UEdGraph* HoverGraph) const override;
+	virtual bool TryCreatingAssetNode(const FAssetData& AssetData, UEdGraph* ParentGraph, const FVector2D Location, EK2NewNodeFlags Options) const override;
+};
+
+FText FInputActionGraphActions::GetGraphHoverMessage(const FAssetData& AssetData, const UEdGraph* HoverGraph) const
+{
+	return FText::Format(LOCTEXT("InputActionHoverMessage", "{0}"), FText::FromName(AssetData.AssetName));
+}
+
+bool FInputActionGraphActions::TryCreatingAssetNode(const FAssetData& AssetData, UEdGraph* ParentGraph, const FVector2D Location, EK2NewNodeFlags Options) const
+{
+	if (AssetData.IsValid())
+	{
+		if (const UInputAction* Action = Cast<const UInputAction>(AssetData.GetAsset()))
+		{
+			for (TObjectPtr<UEdGraphNode> Node : ParentGraph->Nodes)
+			{
+				if(const UK2Node_EnhancedInputAction* InputActionNode = Cast<UK2Node_EnhancedInputAction>(Node))
+				{
+					if (InputActionNode->InputAction.GetFName() == AssetData.AssetName)
+					{
+						if (const TSharedPtr<IBlueprintEditor> BlueprintEditor = FKismetEditorUtilities::GetIBlueprintEditorForObject(ParentGraph, false))
+						{
+							BlueprintEditor.Get()->JumpToPin(InputActionNode->GetPinAt(0));
+						}
+						
+						return false;
+					}
+				}
+			}
+
+			UK2Node_EnhancedInputAction* NewNode = FEdGraphSchemaAction_K2NewNode::SpawnNode<UK2Node_EnhancedInputAction>(
+				ParentGraph,
+				Location,
+				Options,
+				[Action](UK2Node_EnhancedInputAction* NewInstance)
+				{
+					NewInstance->InputAction = Action;
+				}
+
+			);
+			return true;
+		}
+	}
+	return false;
+}
+
+
+
+class FInputBlueprintNodesModule : public IModuleInterface, public FTickableEditorObject
+{
+	// IModuleInterface interface
+	virtual void StartupModule() override
+	{
+		// Register graph actions:
+		FBlueprintGraphModule& GraphModule = FModuleManager::LoadModuleChecked<FBlueprintGraphModule>("BlueprintGraph");
+		{
+			GraphModule.RegisterGraphAction(UInputAction::StaticClass(), MakeUnique<FInputActionGraphActions>());
+		}
+	}
+	// End IModuleInterface interface
+
+	// FTickableEditorObject interface
+	virtual void Tick(float DeltaTime) override
+	{
+		// Update any blueprints that are referencing an input action with a modified value type
+		if (UInputAction::ActionsWithModifiedValueTypes.Num() || UInputAction::ActionsWithModifiedTriggers.Num())
+		{
+			TSet<UBlueprint*> BPsModifiedFromValueTypeChange;
+			TSet<UBlueprint*> BPsModifiedFromTriggerChange;
+			
+			for (TObjectIterator<UK2Node_EnhancedInputAction> NodeIt; NodeIt; ++NodeIt)
+			{
+				if (!FBlueprintNodeTemplateCache::IsTemplateOuter(NodeIt->GetGraph()))
+				{
+					if (UInputAction::ActionsWithModifiedValueTypes.Contains(NodeIt->InputAction))
+					{
+						NodeIt->ReconstructNode();
+						BPsModifiedFromValueTypeChange.Emplace(NodeIt->GetBlueprint());
+					}
+					if (UInputAction::ActionsWithModifiedTriggers.Contains(NodeIt->InputAction))
+					{
+						NodeIt->ReconstructNode();
+						BPsModifiedFromTriggerChange.Emplace(NodeIt->GetBlueprint());
+					}
+				}
+			}
+			for (TObjectIterator<UK2Node_GetInputActionValue> NodeIt; NodeIt; ++NodeIt)
+			{
+				if (!FBlueprintNodeTemplateCache::IsTemplateOuter(NodeIt->GetGraph()))
+				{
+					if (UInputAction::ActionsWithModifiedValueTypes.Contains(NodeIt->InputAction))
+					{
+						NodeIt->ReconstructNode();
+						BPsModifiedFromValueTypeChange.Emplace(NodeIt->GetBlueprint());
+					}
+					if (UInputAction::ActionsWithModifiedTriggers.Contains(NodeIt->InputAction))
+					{
+						NodeIt->ReconstructNode();
+						BPsModifiedFromTriggerChange.Emplace(NodeIt->GetBlueprint());
+					}
+				}
+			}
+
+			if (BPsModifiedFromValueTypeChange.Num())
+			{
+				FNotificationInfo Info(FText::Format(LOCTEXT("ActionValueTypeChange", "Changing action value type affected {0} blueprint(s)!"), BPsModifiedFromValueTypeChange.Num()));
+				Info.ExpireDuration = 5.0f;
+				FSlateNotificationManager::Get().AddNotification(Info);
+			}
+			if (BPsModifiedFromTriggerChange.Num())
+			{
+				FNotificationInfo Info(FText::Format(LOCTEXT("ActionTriggerChange", "Changing action triggers affected {0} blueprint(s)!"), BPsModifiedFromTriggerChange.Num()));
+				Info.ExpireDuration = 5.0f;
+				FSlateNotificationManager::Get().AddNotification(Info);
+			}
+
+			UInputAction::ActionsWithModifiedValueTypes.Reset();
+			UInputAction::ActionsWithModifiedTriggers.Reset();
+		}
+	}
+	
+	virtual TStatId GetStatId() const override { RETURN_QUICK_DECLARE_CYCLE_STAT(FInputBlueprintNodesModule, STATGROUP_Tickables); }
+	// End FTickableEditorObject interface
 };
 
 IMPLEMENT_MODULE(FInputBlueprintNodesModule, InputBlueprintNodes)
diff --git a/Engine/Plugins/EnhancedInput/Source/InputBlueprintNodes/Public/InputBlueprintNodesModule.h b/Engine/Plugins/EnhancedInput/Source/InputBlueprintNodes/Public/InputBlueprintNodesModule.h
--- a/Engine/Plugins/EnhancedInput/Source/InputBlueprintNodes/Public/InputBlueprintNodesModule.h
+++ b/Engine/Plugins/EnhancedInput/Source/InputBlueprintNodes/Public/InputBlueprintNodesModule.h
@@ -2,9 +2,4 @@
 
 #pragma once
 
-//#include "InputBlueprintNodesModule.generated.h"
-
-// Module is not publicly exposed
-
-
-
+// Module is not publicly exposed
\ No newline at end of file
diff --git a/Engine/Plugins/EnhancedInput/Source/InputEditor/InputEditor.Build.cs b/Engine/Plugins/EnhancedInput/Source/InputEditor/InputEditor.Build.cs
--- a/Engine/Plugins/EnhancedInput/Source/InputEditor/InputEditor.Build.cs
+++ b/Engine/Plugins/EnhancedInput/Source/InputEditor/InputEditor.Build.cs
@@ -20,8 +20,7 @@ namespace UnrealBuildTool.Rules
                     "DetailCustomizations",
                     "Engine",
                     "GraphEditor",
-					"InputBlueprintNodes",
-					"InputCore",
+                    "InputCore",
 					"KismetCompiler",
 					"PropertyEditor",
                     "SharedSettingsWidgets",
diff --git a/Engine/Plugins/EnhancedInput/Source/InputEditor/Private/InputEditorModule.cpp b/Engine/Plugins/EnhancedInput/Source/InputEditor/Private/InputEditorModule.cpp
--- a/Engine/Plugins/EnhancedInput/Source/InputEditor/Private/InputEditorModule.cpp
+++ b/Engine/Plugins/EnhancedInput/Source/InputEditor/Private/InputEditorModule.cpp
@@ -4,19 +4,15 @@
 
 #include "AssetBlueprintGraphActions.h"
 #include "BlueprintEditorModule.h"
-#include "BlueprintGraphModule.h"
 #include "BlueprintNodeTemplateCache.h"
 #include "Brushes/SlateImageBrush.h"
 #include "Framework/Notifications/NotificationManager.h"
-#include "Engine/Blueprint.h"
 #include "InputMappingContext.h"
 #include "Misc/PackageName.h"
 #include "PlayerMappableInputConfig.h"
 #include "InputCustomizations.h"
 #include "ISettingsModule.h"
 #include "ToolMenuSection.h"
-#include "K2Node_EnhancedInputAction.h"
-#include "K2Node_GetInputActionValue.h"
 #include "PropertyEditorModule.h"
 #include "UObject/UObjectIterator.h"
 #include "Widgets/Notifications/SNotificationList.h"
@@ -30,7 +26,6 @@
 #include "IContentBrowserSingleton.h"
 #include "ClassViewerModule.h"
 #include "ClassViewerFilter.h"
-#include "Kismet2/KismetEditorUtilities.h"
 #include "Kismet2/SClassPickerDialog.h"
 #include "GameFramework/InputSettings.h"
 #include "EnhancedInputComponent.h"
@@ -38,8 +33,9 @@
 #include "Interfaces/IMainFrameModule.h"
 #include "SourceControlHelpers.h"
 #include "HAL/FileManager.h"
-#include "Kismet2/KismetEditorUtilities.h"
 #include "PlayerMappableKeySettings.h"
+#include "IAssetTools.h"
+#include "IAssetTypeActions.h"
 
 #include UE_INLINE_GENERATED_CPP_BY_NAME(InputEditorModule)
 
@@ -327,55 +323,6 @@ public:
 };
 PRAGMA_ENABLE_DEPRECATION_WARNINGS
 
-struct FInputActionGraphActions : public FAssetBlueprintGraphActions
-{
-	virtual FText GetGraphHoverMessage(const FAssetData& AssetData, const UEdGraph* HoverGraph) const override;
-	virtual bool TryCreatingAssetNode(const FAssetData& AssetData, UEdGraph* ParentGraph, const FVector2D Location, EK2NewNodeFlags Options) const override;
-};
-
-FText FInputActionGraphActions::GetGraphHoverMessage(const FAssetData& AssetData, const UEdGraph* HoverGraph) const
-{
-	return FText::Format(LOCTEXT("InputActionHoverMessage", "{0}"), FText::FromName(AssetData.AssetName));
-}
-
-bool FInputActionGraphActions::TryCreatingAssetNode(const FAssetData& AssetData, UEdGraph* ParentGraph, const FVector2D Location, EK2NewNodeFlags Options) const
-{
-	if (AssetData.IsValid())
-	{
-		if (const UInputAction* Action = Cast<const UInputAction>(AssetData.GetAsset()))
-		{
-			for (TObjectPtr<UEdGraphNode> Node : ParentGraph->Nodes)
-			{
-				if(const UK2Node_EnhancedInputAction* InputActionNode = Cast<UK2Node_EnhancedInputAction>(Node))
-				{
-					if (InputActionNode->InputAction.GetFName() == AssetData.AssetName)
-					{
-						if (const TSharedPtr<IBlueprintEditor> BlueprintEditor = FKismetEditorUtilities::GetIBlueprintEditorForObject(ParentGraph, false))
-						{
-							BlueprintEditor.Get()->JumpToPin(InputActionNode->GetPinAt(0));
-						}
-						
-						return false;
-					}
-				}
-			}
-
-			UK2Node_EnhancedInputAction* NewNode = FEdGraphSchemaAction_K2NewNode::SpawnNode<UK2Node_EnhancedInputAction>(
-				ParentGraph,
-				Location,
-				Options,
-				[Action](UK2Node_EnhancedInputAction* NewInstance)
-				{
-					NewInstance->InputAction = Action;
-				}
-
-			);
-			return true;
-		}
-	}
-	return false;
-}
-
 /** Custom style set for Enhanced Input */
 class FEnhancedInputSlateStyle final : public FSlateStyleSet
 {
@@ -426,15 +373,6 @@ void FInputEditorModule::StartupModule()
 		PRAGMA_DISABLE_DEPRECATION_WARNINGS
 		RegisterAssetTypeActions(AssetTools, MakeShareable(new FAssetTypeActions_PlayerMappableInputConfig));
 		PRAGMA_ENABLE_DEPRECATION_WARNINGS
-		// TODO: Build these off a button on the InputContext Trigger/Mapping pickers? Would be good to have both.
-		//RegisterAssetTypeActions(AssetTools, MakeShareable(new FAssetTypeActions_InputTrigger));
-		//RegisterAssetTypeActions(AssetTools, MakeShareable(new FAssetTypeActions_InputModifier));
-	}
-
-	// Register graph actions:
-	FBlueprintGraphModule& GraphModule = FModuleManager::LoadModuleChecked<FBlueprintGraphModule>("BlueprintGraph");
-	{
-		GraphModule.RegisterGraphAction(UInputAction::StaticClass(), MakeUnique<FInputActionGraphActions>());
 	}
 
 	// Make a new style set for Enhanced Input, which will register any custom icons for the types in this plugin
@@ -486,6 +424,11 @@ void FInputEditorModule::ShutdownModule()
 	}
 }
 
+EAssetTypeCategories::Type FInputEditorModule::GetInputAssetsCategory()
+{
+	return InputAssetsCategory;
+}
+
 void FInputEditorModule::OnMainFrameCreationFinished(TSharedPtr<SWindow> InRootWindow, bool bIsRunningStartupDialog)
 {
 	if (FApp::HasProjectName())
@@ -582,63 +525,10 @@ bool FInputEditorModule::IsMappingNameInUse(const FName InName)
 	return FindMappingByName(InName) != nullptr;
 }
 
-void FInputEditorModule::Tick(float DeltaTime)
+void FInputEditorModule::RegisterAssetTypeActions(IAssetTools& AssetTools, TSharedRef<IAssetTypeActions> Action)
 {
-	// Update any blueprints that are referencing an input action with a modified value type
-	if (UInputAction::ActionsWithModifiedValueTypes.Num() || UInputAction::ActionsWithModifiedTriggers.Num())
-	{
-		TSet<UBlueprint*> BPsModifiedFromValueTypeChange;
-		TSet<UBlueprint*> BPsModifiedFromTriggerChange;
-		
-		for (TObjectIterator<UK2Node_EnhancedInputAction> NodeIt; NodeIt; ++NodeIt)
-		{
-			if (!FBlueprintNodeTemplateCache::IsTemplateOuter(NodeIt->GetGraph()))
-			{
-				if (UInputAction::ActionsWithModifiedValueTypes.Contains(NodeIt->InputAction))
-				{
-					NodeIt->ReconstructNode();
-					BPsModifiedFromValueTypeChange.Emplace(NodeIt->GetBlueprint());
-				}
-				if (UInputAction::ActionsWithModifiedTriggers.Contains(NodeIt->InputAction))
-				{
-					NodeIt->ReconstructNode();
-					BPsModifiedFromTriggerChange.Emplace(NodeIt->GetBlueprint());
-				}
-			}
-		}
-		for (TObjectIterator<UK2Node_GetInputActionValue> NodeIt; NodeIt; ++NodeIt)
-		{
-			if (!FBlueprintNodeTemplateCache::IsTemplateOuter(NodeIt->GetGraph()))
-			{
-				if (UInputAction::ActionsWithModifiedValueTypes.Contains(NodeIt->InputAction))
-				{
-					NodeIt->ReconstructNode();
-					BPsModifiedFromValueTypeChange.Emplace(NodeIt->GetBlueprint());
-				}
-				if (UInputAction::ActionsWithModifiedTriggers.Contains(NodeIt->InputAction))
-				{
-					NodeIt->ReconstructNode();
-					BPsModifiedFromTriggerChange.Emplace(NodeIt->GetBlueprint());
-				}
-			}
-		}
-
-		if (BPsModifiedFromValueTypeChange.Num())
-		{
-			FNotificationInfo Info(FText::Format(LOCTEXT("ActionValueTypeChange", "Changing action value type affected {0} blueprint(s)!"), BPsModifiedFromValueTypeChange.Num()));
-			Info.ExpireDuration = 5.0f;
-			FSlateNotificationManager::Get().AddNotification(Info);
-		}
-		if (BPsModifiedFromTriggerChange.Num())
-		{
-			FNotificationInfo Info(FText::Format(LOCTEXT("ActionTriggerChange", "Changing action triggers affected {0} blueprint(s)!"), BPsModifiedFromTriggerChange.Num()));
-			Info.ExpireDuration = 5.0f;
-			FSlateNotificationManager::Get().AddNotification(Info);
-		}
-
-		UInputAction::ActionsWithModifiedValueTypes.Reset();
-		UInputAction::ActionsWithModifiedTriggers.Reset();
-	}
+	AssetTools.RegisterAssetTypeActions(Action);
+	CreatedAssetTypeActions.Add(Action);
 }
 
 #undef LOCTEXT_NAMESPACE
\ No newline at end of file
diff --git a/Engine/Plugins/EnhancedInput/Source/InputEditor/Public/InputEditorModule.h b/Engine/Plugins/EnhancedInput/Source/InputEditor/Public/InputEditorModule.h
--- a/Engine/Plugins/EnhancedInput/Source/InputEditor/Public/InputEditorModule.h
+++ b/Engine/Plugins/EnhancedInput/Source/InputEditor/Public/InputEditorModule.h
@@ -3,9 +3,8 @@
 #pragma once
 
 #include "Factories/Factory.h"
-#include "IAssetTools.h"
-#include "IAssetTypeActions.h"
-#include "TickableEditorObject.h"
+#include "Modules/ModuleInterface.h"
+#include "AssetTypeCategories.h"
 
 #include "InputEditorModule.generated.h"
 
@@ -17,8 +16,11 @@ DECLARE_LOG_CATEGORY_EXTERN(LogEnhancedInputEditor, Log, All);
 class UInputAction;
 class SWindow;
 class UPlayerMappableKeySettings;
+class IAssetTypeActions;
+class IAssetTools;
+class FSlateStyleSet;
 
-class FInputEditorModule : public IModuleInterface, public FTickableEditorObject
+class FInputEditorModule : public IModuleInterface
 {
 public:
 
@@ -27,12 +29,7 @@ public:
 	virtual void ShutdownModule() override;
 	// End IModuleInterface interface
 
-	// FTickableEditorObject interface
-	virtual void Tick(float DeltaTime) override;
-	virtual TStatId GetStatId() const override { RETURN_QUICK_DECLARE_CYCLE_STAT(FInputEditorModule, STATGROUP_Tickables); }
-	// End FTickableEditorObject interface
-
-	static EAssetTypeCategories::Type GetInputAssetsCategory() { return InputAssetsCategory; }
+	static EAssetTypeCategories::Type GetInputAssetsCategory();
 	
 	/** Returns a pointer to the player mappable settings object that has this mapping name */
 	INPUTEDITOR_API static const UPlayerMappableKeySettings* FindMappingByName(const FName InName);
@@ -41,11 +38,7 @@ public:
 	INPUTEDITOR_API static bool IsMappingNameInUse(const FName InName);
 		
 private:
-	void RegisterAssetTypeActions(IAssetTools& AssetTools, TSharedRef<IAssetTypeActions> Action)
-	{
-		AssetTools.RegisterAssetTypeActions(Action);
-		CreatedAssetTypeActions.Add(Action);
-	}
+	void RegisterAssetTypeActions(IAssetTools& AssetTools, TSharedRef<IAssetTypeActions> Action);
 
 	void OnMainFrameCreationFinished(TSharedPtr<SWindow> InRootWindow, bool bIsRunningStartupDialog);
 	
@@ -56,7 +49,7 @@ private:
 	
 	TArray<TSharedPtr<IAssetTypeActions>> CreatedAssetTypeActions;
 	
-	TSharedPtr<class FSlateStyleSet> StyleSet;
+	TSharedPtr<FSlateStyleSet> StyleSet;
 };
 
 ////////////////////////////////////////////////////////////////////
