Mark all exec pins as advanced view on Input Action event nodes except for "Triggered". Most of the time the triggered event pin is what users should be connecting to, and making the other pins hidden will improve first time user UX and make it more obvious what pin you should be hooking up to.

There is an editor preference to change what pins show up by default in blueprints, but this change will hopefully point people towards using Input Triggers more instead of relying on undefined input behaviors and improve the first time user experience.

diff --git a/Engine/Plugins/EnhancedInput/Source/InputBlueprintNodes/Private/K2Node_EnhancedInputAction.cpp b/Engine/Plugins/EnhancedInput/Source/InputBlueprintNodes/Private/K2Node_EnhancedInputAction.cpp
--- a/Engine/Plugins/EnhancedInput/Source/InputBlueprintNodes/Private/K2Node_EnhancedInputAction.cpp
+++ b/Engine/Plugins/EnhancedInput/Source/InputBlueprintNodes/Private/K2Node_EnhancedInputAction.cpp
@@ -20,6 +20,7 @@
 #include "Subsystems/AssetEditorSubsystem.h"
 #include "Modules/ModuleManager.h"
 #include "Styling/AppStyle.h"
+#include "EnhancedInputEditorSettings.h"
 
 #include UE_INLINE_GENERATED_CPP_BY_NAME(K2Node_EnhancedInputAction)
 
@@ -68,6 +69,11 @@ void UK2Node_EnhancedInputAction::AllocateDefaultPins()
 		static const UEnum* EventEnum = StaticEnum<ETriggerEvent>();
 
 		UEdGraphPin* NewPin = CreatePin(EGPD_Output, UEdGraphSchema_K2::PC_Exec, PinName);
+
+		// Mark all triggering exec pins as advanced view except for the triggered pin. Most of the time, triggered is what users should be using.
+		// More advanced input set ups can use the more advanced pins when they want to! 
+		NewPin->bAdvancedView = !(GetDefault<UEnhancedInputEditorSettings>()->VisibleEventPinsByDefault & static_cast<uint8>(Event));
+
 		NewPin->PinToolTip = EventEnum->GetToolTipTextByIndex(EventEnum->GetIndexByValue(static_cast<uint8>(Event))).ToString();
 
 		// Add a special tooltip and display name for pins that are unsupported
@@ -96,6 +102,7 @@ void UK2Node_EnhancedInputAction::AllocateDefaultPins()
 		UEdGraphPin* ActionPin = CreatePin(EGPD_Output, UEdGraphSchema_K2::PC_Object, InputAction->GetClass(), InputActionPinName);
 		ActionPin->DefaultObject = const_cast<UObject*>(Cast<UObject>(InputAction));
 		ActionPin->DefaultValue = InputAction->GetName();
+		ActionPin->bAdvancedView = true;
 		Schema->ConstructBasicPinTooltip(*ActionPin, LOCTEXT("InputActionPinDescription", "The input action that caused this event to fire"), ActionPin->PinToolTip);	
 	}
 	
@@ -114,7 +121,7 @@ void UK2Node_EnhancedInputAction::HideEventPins(UEdGraphPin* RetainPin)
 		{
 			const bool bIsSupported = UInputTrigger::IsSupportedTriggerEvent(SupportedTriggerEvents, Event);
 			
-			Pin->bAdvancedView = !bIsSupported;
+			Pin->bAdvancedView = !(GetDefault<UEnhancedInputEditorSettings>()->VisibleEventPinsByDefault & static_cast<uint8>(Event)) || !bIsSupported;
 		}
 	});
 }
@@ -201,9 +208,10 @@ FText UK2Node_EnhancedInputAction::GetTooltipText() const
 		FString ActionPath = InputAction ? InputAction->GetFullName() : TEXT("");
 		CachedTooltip.SetCachedText(
 			FText::Format(
-				LOCTEXT("EnhancedInputAction_Tooltip", "Event for when '{0}' triggers.\n\nNote: This is not guaranteed to fire every frame, only when the Action is triggered and the current Input Mode includes 'Game'.\n\n{1}"),
-			FText::FromString(ActionPath),
-			LOCTEXT("EnhancedInputAction_Node_Tooltip_Tip", "Tip: Use the 'showdebug enhancedinput' command while playing to see debug information about Enhanced Input.")),
+				LOCTEXT("EnhancedInputAction_Tooltip", "Event for when '{0}' triggers.\n\nNote: This is not guaranteed to fire every frame, only when the Action is triggered and the current Input Mode includes 'Game'.\n\n{1}\n\n{2}"),
+				FText::FromString(ActionPath),
+				LOCTEXT("EnhancedInputAction_Node_Tooltip_Tip", "Tip: Use the 'showdebug enhancedinput' command while playing to see debug information about Enhanced Input."),
+				LOCTEXT("EnhancedInputAction_Node_SettingsTooltip", "You can change what execution pins are visible by default in the Enhanced Input Editor Preferences.")),
 			this);
 	}
 	return CachedTooltip;
diff --git a/Engine/Plugins/EnhancedInput/Source/InputEditor/Private/EnhancedInputEditorSettings.cpp b/Engine/Plugins/EnhancedInput/Source/InputEditor/Private/EnhancedInputEditorSettings.cpp
--- a/Engine/Plugins/EnhancedInput/Source/InputEditor/Private/EnhancedInputEditorSettings.cpp
+++ b/Engine/Plugins/EnhancedInput/Source/InputEditor/Private/EnhancedInputEditorSettings.cpp
@@ -40,6 +40,8 @@ UEnhancedInputEditorProjectSettings::UEnhancedInputEditorProjectSettings(const F
 UEnhancedInputEditorSettings::UEnhancedInputEditorSettings()
 	: bLogAllInput(false)
 	, bAutomaticallyStartConsumingInput(false)
+	// By default only show the triggered event 
+	, VisibleEventPinsByDefault(static_cast<uint8>(ETriggerEvent::Triggered))
 {
 }
 
diff --git a/Engine/Plugins/EnhancedInput/Source/InputEditor/Public/EnhancedInputEditorSettings.h b/Engine/Plugins/EnhancedInput/Source/InputEditor/Public/EnhancedInputEditorSettings.h
--- a/Engine/Plugins/EnhancedInput/Source/InputEditor/Public/EnhancedInputEditorSettings.h
+++ b/Engine/Plugins/EnhancedInput/Source/InputEditor/Public/EnhancedInputEditorSettings.h
@@ -31,10 +31,10 @@ public:
 };
 
 /**
- * A collection of useful indivudal user settings when using the EnhancedInputEditorSubsystem.
+ * A collection of useful individual user settings when using the EnhancedInputEditorSubsystem.
  */
 UCLASS(config=EditorPerProjectUserSettings, meta=(DisplayName="Enhanced Input Editor Settings"))
-class UEnhancedInputEditorSettings : public UDeveloperSettingsBackedByCVars
+class INPUTEDITOR_API UEnhancedInputEditorSettings : public UDeveloperSettingsBackedByCVars
 {
 	GENERATED_BODY()
 public:
@@ -51,6 +51,10 @@ public:
 	/** If true, then the UEnhancedInputEditorSubsystem will be started when it is initalized */
 	UPROPERTY(config, EditAnywhere, Category = Editor, meta=(ConsoleVariable="EnhancedEditorInput.bAutomaticallyStartConsumingInput"))
 	uint8 bAutomaticallyStartConsumingInput : 1;
+
+	/** A bitmask of what event pins are visible when you place an Input Action event node in blueprints.  */
+	UPROPERTY(config, EditAnywhere, Category = Blueprints, meta = (Bitmask, BitmaskEnum = "/Script/EnhancedInput.ETriggerEvent"))
+	uint8 VisibleEventPinsByDefault;
 };
 
 #if UE_ENABLE_INCLUDE_ORDER_DEPRECATED_IN_5_2
