Updating enhanced input settings to allow deferral of settings updates.

diff --git a/Engine/Plugins/EnhancedInput/Source/EnhancedInput/Private/UserSettings/EnhancedInputUserSettings.cpp b/Engine/Plugins/EnhancedInput/Source/EnhancedInput/Private/UserSettings/EnhancedInputUserSettings.cpp
--- a/Engine/Plugins/EnhancedInput/Source/EnhancedInput/Private/UserSettings/EnhancedInputUserSettings.cpp
+++ b/Engine/Plugins/EnhancedInput/Source/EnhancedInput/Private/UserSettings/EnhancedInputUserSettings.cpp
@@ -9,6 +9,7 @@
 #include "EnhancedActionKeyMapping.h"
 #include "EnhancedInputSubsystems.h"
 #include "Engine/LocalPlayer.h"
+#include "Engine/World.h"
 #include "Kismet/GameplayStatics.h"
 #include "EnhancedInputDeveloperSettings.h"
 #include "Internationalization/Text.h"
@@ -86,6 +87,7 @@ FMapPlayerKeyArgs::FMapPlayerKeyArgs()
 	, NewKey(EKeys::Invalid)
 	, HardwareDeviceId(NAME_None)
 	, bCreateMatchingSlotIfNeeded(true)
+	, bDeferOnSettingsChangedBroadcast (false)
 {
 }
 
@@ -864,7 +866,26 @@ void UEnhancedInputUserSettings::MapPlayerKey(const FMapPlayerKeyArgs& InArgs, F
 		FoundMapping->SetCurrentKey(InArgs.NewKey);
 		OnKeyMappingUpdated(FoundMapping, InArgs, false);
 		
-		OnSettingsChanged.Broadcast(this);
+		if(InArgs.bDeferOnSettingsChangedBroadcast  && !DeferredSettingsChangedTimerHandle.IsValid())
+		{
+			if(UWorld* World = GetWorld())
+			{
+				TWeakObjectPtr<UEnhancedInputUserSettings> WeakThis = this;
+				DeferredSettingsChangedTimerHandle = World->GetTimerManager().SetTimerForNextTick(FTimerDelegate::CreateWeakLambda(this, 
+					[WeakThis]
+					{
+						if(WeakThis.IsValid())
+						{
+							WeakThis->OnSettingsChanged.Broadcast(WeakThis.Get());
+							WeakThis->DeferredSettingsChangedTimerHandle.Invalidate();
+						}
+					}));
+			}
+		}
+		else
+		{
+			OnSettingsChanged.Broadcast(this);
+		}
 	}
 	// If it doesn't exist, then we need to make it if there is a valid action name
 	else if (FKeyMappingRow* MappingRow = KeyProfile->FindKeyMappingRowMutable(InArgs.MappingName))
@@ -911,7 +932,26 @@ void UEnhancedInputUserSettings::MapPlayerKey(const FMapPlayerKeyArgs& InArgs, F
 			const FSetElementId SetElem = MappingRow->Mappings.Add(PlayerMappingData);
 			OnKeyMappingUpdated(&MappingRow->Mappings.Get(SetElem), InArgs, false);
 			
-			OnSettingsChanged.Broadcast(this);
+			if(InArgs.bDeferOnSettingsChangedBroadcast  && !DeferredSettingsChangedTimerHandle.IsValid())
+			{
+				if(UWorld* World = GetWorld())
+				{
+					TWeakObjectPtr<UEnhancedInputUserSettings> WeakThis = this;
+					DeferredSettingsChangedTimerHandle = World->GetTimerManager().SetTimerForNextTick(FTimerDelegate::CreateWeakLambda(this, 
+						[WeakThis]
+						{
+							if(WeakThis.IsValid())
+							{
+								WeakThis->OnSettingsChanged.Broadcast(WeakThis.Get());
+								WeakThis->DeferredSettingsChangedTimerHandle.Invalidate();
+							}
+						}));
+				}
+			}
+			else
+			{
+				OnSettingsChanged.Broadcast(this);
+			}
 		}
 	}
 	else
diff --git a/Engine/Plugins/EnhancedInput/Source/EnhancedInput/Public/UserSettings/EnhancedInputUserSettings.h b/Engine/Plugins/EnhancedInput/Source/EnhancedInput/Public/UserSettings/EnhancedInputUserSettings.h
--- a/Engine/Plugins/EnhancedInput/Source/EnhancedInput/Public/UserSettings/EnhancedInputUserSettings.h
+++ b/Engine/Plugins/EnhancedInput/Source/EnhancedInput/Public/UserSettings/EnhancedInputUserSettings.h
@@ -80,6 +80,10 @@ struct ENHANCEDINPUT_API FMapPlayerKeyArgs final
 	/** If there is not a player mapping already with the same Slot and Hardware Device ID, then create a new mapping for this slot. */
 	UPROPERTY(EditAnywhere, BlueprintReadWrite, Category="Enhanced Input|User Settings")
 	uint8 bCreateMatchingSlotIfNeeded : 1;
+
+	/** Defers setting changed delegates until the next frame if set to true. */
+	UPROPERTY(EditAnywhere, BlueprintReadWrite, Category="Enhanced Input|User Settings")
+	uint8 bDeferOnSettingsChangedBroadcast : 1;
 };
 
 /** Represents a single key mapping that is set by the player */
@@ -536,6 +540,11 @@ protected:
 public:
 	
 	ULocalPlayer* GetLocalPlayer() const;
+
+protected:
+	// Used to track when a settings change callback has been deferred till the next frame.
+	FTimerHandle DeferredSettingsChangedTimerHandle;
+public:
 	
 	/** Fired when the user settings have changed, such as their key mappings. */
 	DECLARE_DYNAMIC_MULTICAST_DELEGATE_OneParam(FEnhancedInputUserSettingsChanged, UEnhancedInputUserSettings*, Settings);
