Missing include in previous change.

diff --git a/Engine/Plugins/EnhancedInput/Source/EnhancedInput/Private/UserSettings/EnhancedInputUserSettings.cpp b/Engine/Plugins/EnhancedInput/Source/EnhancedInput/Private/UserSettings/EnhancedInputUserSettings.cpp
--- a/Engine/Plugins/EnhancedInput/Source/EnhancedInput/Private/UserSettings/EnhancedInputUserSettings.cpp
+++ b/Engine/Plugins/EnhancedInput/Source/EnhancedInput/Private/UserSettings/EnhancedInputUserSettings.cpp
@@ -2,25 +2,26 @@
 
 #include "UserSettings/EnhancedInputUserSettings.h"
 
-#include "EnhancedInputModule.h"
-#include "InputMappingContext.h"
-#include "InputAction.h"
-#include "PlayerMappableKeySettings.h"
-#include "EnhancedActionKeyMapping.h"
-#include "EnhancedInputSubsystems.h"
+#include "Algo/Find.h"
+#include "DrawDebugHelpers.h"	// required for ENABLE_DRAW_DEBUG define
+#include "Engine/Engine.h"		// for GEngine definition
 #include "Engine/LocalPlayer.h"
 #include "Engine/World.h"
-#include "Kismet/GameplayStatics.h"
+#include "EnhancedActionKeyMapping.h"
 #include "EnhancedInputDeveloperSettings.h"
-#include "Internationalization/Text.h"
+#include "EnhancedInputLibrary.h"
+#include "EnhancedInputModule.h"
+#include "EnhancedInputSubsystems.h"
 #include "GameFramework/InputSettings.h"
 #include "HAL/IConsoleManager.h"
-#include "EnhancedInputLibrary.h"
-#include "Algo/Find.h"
+#include "InputAction.h"
+#include "InputMappingContext.h"
+#include "Internationalization/Text.h"
+#include "Kismet/GameplayStatics.h"
 #include "NativeGameplayTags.h"
+#include "PlayerMappableKeySettings.h"
 #include "SaveGameSystem.h"
-#include "Engine/Engine.h"		// for GEngine definition
-#include "DrawDebugHelpers.h"	// required for ENABLE_DRAW_DEBUG define
+#include "TimerManager.h"
 
 #if ENABLE_DRAW_DEBUG
 #include "Engine/Canvas.h"
