Add Continuous Input Injection to enhanced input. Before, if you called InjectInputForAction then it would need to be called every tick. A lot of games disable widget tick for perf reasons, so this gives an alternative!

diff --git a/Engine/Plugins/EnhancedInput/Source/EnhancedInput/Private/EnhancedInputSubsystemInterface.cpp b/Engine/Plugins/EnhancedInput/Source/EnhancedInput/Private/EnhancedInputSubsystemInterface.cpp
--- a/Engine/Plugins/EnhancedInput/Source/EnhancedInput/Private/EnhancedInputSubsystemInterface.cpp
+++ b/Engine/Plugins/EnhancedInput/Source/EnhancedInput/Private/EnhancedInputSubsystemInterface.cpp
@@ -93,6 +93,58 @@ void IEnhancedInputSubsystemInterface::InjectInputForPlayerMapping(const FName M
 	InjectInputVectorForPlayerMapping(MappingName, RawValue.Get<FVector>(), Modifiers, Triggers);
 }
 
+void IEnhancedInputSubsystemInterface::StartContinuousInputInjectionForAction(const UInputAction* Action, FInputActionValue RawValue, const TArray<UInputModifier*>& Modifiers, const TArray<UInputTrigger*>& Triggers)
+{
+	FInjectedInput& Injection = ContinuouslyInjectedInputs.FindOrAdd(Action);
+	
+	Injection.RawValue = RawValue;
+	Injection.Modifiers = Modifiers;
+	Injection.Triggers = Triggers;
+}
+
+void IEnhancedInputSubsystemInterface::StartContinuousInputInjectionForPlayerMapping(const FName MappingName, FInputActionValue RawValue, const TArray<UInputModifier*>& Modifiers, const TArray<UInputTrigger*>& Triggers)
+{
+	if (const UEnhancedInputUserSettings* UserSettings = GetUserSettings())
+	{
+		if (const UInputAction* Action = UserSettings->FindInputActionForMapping(MappingName))
+		{
+			StartContinuousInputInjectionForAction(Action, RawValue, Modifiers, Triggers);
+		}
+		else
+		{
+			UE_LOG(LogEnhancedInput, Warning, TEXT("Could not find a Input Action for mapping name '%s'"), *MappingName.ToString());
+		}
+	}
+	else
+	{
+		UE_LOG(LogEnhancedInput, Warning, TEXT("Could not find a valid UEnhancedInputUserSettings object, is it enabled in the project settings?"));
+	}
+}
+
+void IEnhancedInputSubsystemInterface::StopContinuousInputInjectionForAction(const UInputAction* Action)
+{
+	ContinuouslyInjectedInputs.Remove(Action);
+}
+
+void IEnhancedInputSubsystemInterface::StopContinuousInputInjectionForPlayerMapping(const FName MappingName)
+{
+	if (const UEnhancedInputUserSettings* UserSettings = GetUserSettings())
+	{
+		if (const UInputAction* Action = UserSettings->FindInputActionForMapping(MappingName))
+		{
+			StopContinuousInputInjectionForAction(Action);
+		}
+		else
+		{
+			UE_LOG(LogEnhancedInput, Warning, TEXT("Could not find a Input Action for mapping name '%s'"), *MappingName.ToString());
+		}
+	}
+	else
+	{
+		UE_LOG(LogEnhancedInput, Warning, TEXT("Could not find a valid UEnhancedInputUserSettings object, is it enabled in the project settings?"));
+	}
+}
+
 void IEnhancedInputSubsystemInterface::InjectInputVectorForPlayerMapping(const FName MappingName, FVector Value, const TArray<UInputModifier*>& Modifiers, const TArray<UInputTrigger*>& Triggers)
 {
 	if (const UEnhancedInputUserSettings* UserSettings = GetUserSettings())
@@ -991,6 +1043,16 @@ void IEnhancedInputSubsystemInterface::TickForcedInput(float DeltaTime)
 		return;
 	}
 
+	// Any continuous input injection needs to be added each frame until its stopped
+	for (TPair<TObjectPtr<const UInputAction>, FInjectedInput>& ContinuousInjection : ContinuouslyInjectedInputs)
+	{
+		TObjectPtr<const UInputAction>& Action = ContinuousInjection.Key;
+		if (const UInputAction* InputAction = Action.Get())
+		{
+			PlayerInput->InjectInputForAction(InputAction, ContinuousInjection.Value.RawValue, ContinuousInjection.Value.Modifiers, ContinuousInjection.Value.Triggers);
+		}
+	}
+
 	// Forced action triggering
 	for (TPair<TWeakObjectPtr<const UInputAction>, FInputActionValue>& ForcedActionPair : ForcedActions)
 	{
diff --git a/Engine/Plugins/EnhancedInput/Source/EnhancedInput/Public/EnhancedInputSubsystemInterface.h b/Engine/Plugins/EnhancedInput/Source/EnhancedInput/Public/EnhancedInputSubsystemInterface.h
--- a/Engine/Plugins/EnhancedInput/Source/EnhancedInput/Public/EnhancedInputSubsystemInterface.h
+++ b/Engine/Plugins/EnhancedInput/Source/EnhancedInput/Public/EnhancedInputSubsystemInterface.h
@@ -145,6 +145,44 @@ public:
 	 */
 	UFUNCTION(BlueprintCallable, Category="Input", meta=(AutoCreateRefTerm="Modifiers,Triggers"))
 	virtual void InjectInputVectorForPlayerMapping(const FName MappingName, FVector Value, const TArray<UInputModifier*>& Modifiers, const TArray<UInputTrigger*>& Triggers);
+
+	/**
+	 * Starts simulation of input via injection. This injects the given input every tick until it is stopped with StopContinuousInputInjectionForAction.
+	 *
+	 * @param Action		The Input Action to set inject input for
+	 * @param RawValue		The value to set the action to (the type will be controlled by the Action)
+	 * @param Modifiers		The modifiers to apply to the injected input.
+	 * @param Triggers		The triggers to apply to the injected input.
+	 */
+	UFUNCTION(BlueprintCallable, Category="Input", meta=(AutoCreateRefTerm="Modifiers,Triggers"))
+	virtual void StartContinuousInputInjectionForAction(const UInputAction* Action, FInputActionValue RawValue, const TArray<UInputModifier*>& Modifiers, const TArray<UInputTrigger*>& Triggers);
+
+	/**
+	 * Starts simulation of input via injection. This injects the given input every tick until it is stopped with StopContinuousInputInjectionForAction.
+	 *
+	 * @param MappingName		The name of the player mapping that can be used for look up an associated UInputAction object.
+	 * @param RawValue			The value to set the action to (the type will be controlled by the Action)
+	 * @param Modifiers			The modifiers to apply to the injected input.
+	 * @param Triggers			The triggers to apply to the injected input.
+	 */
+	UFUNCTION(BlueprintCallable, Category="Input", meta=(AutoCreateRefTerm="Modifiers,Triggers"))
+	virtual void StartContinuousInputInjectionForPlayerMapping(const FName MappingName, FInputActionValue RawValue, const TArray<UInputModifier*>& Modifiers, const TArray<UInputTrigger*>& Triggers);
+
+	/**
+	 * Stops continuous input injection for the given action.
+	 *
+	 * @param Action		The action to stop injecting input for
+	 */
+	UFUNCTION(BlueprintCallable, Category="Input")
+	virtual void StopContinuousInputInjectionForAction(const UInputAction* Action);
+
+	/**
+	 * Stops continuous input injection for the given player mapping name.
+	 *
+	 * @param MappingName		The name of the player mapping that can be used for look up an associated UInputAction object.
+	 */
+	UFUNCTION(BlueprintCallable, Category="Input")
+	virtual void StopContinuousInputInjectionForPlayerMapping(const FName MappingName);
 	
 	/**
 	 * Remove all applied mapping contexts.
@@ -402,6 +440,9 @@ private:
 	/** A map of any player mapped keys to the keys that they should redirect to instead */
 	TMap<FName, TMap<FPlayerMappableKeySlot, FKey>> PlayerMappedSettings;
 
+	/** A map of inputs that should be injected every frame. This inputs will be injected when ForcedInput is ticked. */
+	TMap<TObjectPtr<const UInputAction>, FInjectedInput> ContinuouslyInjectedInputs;
+
 	/**
 	 * A map of input actions with a Chorded trigger, mapped to the action they are dependent on.
 	 * The Key is the Input Action with the Chorded Trigger, and the value is the action it is dependant on.
