Make the Enhanced Input behavior when input is flushed consistent with the legacy system.

This makes the behavior of Enhanced Input consistent with the Legacy input system when UPlayerInput::FlushPressedKeys is called (loss of focus to the viewport, or it is otherwise called from gameplay code).

Enhanced Input was incorrectly determing if a key was down on the previous frame if the input was flushed because its RawValue would be reset. We now check if the input is currently being flushed on this frame, and if it is then we can determine if the key was down on the previous frame based on it's "Value" in the key state map. If the "Value" is non-zero, but the "RawValue" has been reset to zero, we know it has been reset from the FlushPressedKeys function.

This behavior is Default ON, but can be toggled off in the Project Settings > Enhanced Input.

diff --git a/Engine/Plugins/EnhancedInput/Source/EnhancedInput/Private/EnhancedInputDeveloperSettings.cpp b/Engine/Plugins/EnhancedInput/Source/EnhancedInput/Private/EnhancedInputDeveloperSettings.cpp
--- a/Engine/Plugins/EnhancedInput/Source/EnhancedInput/Private/EnhancedInputDeveloperSettings.cpp
+++ b/Engine/Plugins/EnhancedInput/Source/EnhancedInput/Private/EnhancedInputDeveloperSettings.cpp
@@ -26,6 +26,7 @@ UEnhancedInputDeveloperSettings::UEnhancedInputDeveloperSettings(const FObjectIn
 	, UserSettingsClass(UEnhancedInputUserSettings::StaticClass())
 	, DefaultPlayerMappableKeyProfileClass(UEnhancedPlayerMappableKeyProfile::StaticClass())
 	, DefaultWorldInputClass(UEnhancedPlayerInput::StaticClass())
+	, bSendTriggeredEventsWhenInputIsFlushed(true)
 	, bEnableUserSettings(false)
 	, bEnableDefaultMappingContexts(true)
 	, bShouldOnlyTriggerLastActionInChord(true)
diff --git a/Engine/Plugins/EnhancedInput/Source/EnhancedInput/Private/EnhancedPlayerInput.cpp b/Engine/Plugins/EnhancedInput/Source/EnhancedInput/Private/EnhancedPlayerInput.cpp
--- a/Engine/Plugins/EnhancedInput/Source/EnhancedInput/Private/EnhancedPlayerInput.cpp
+++ b/Engine/Plugins/EnhancedInput/Source/EnhancedInput/Private/EnhancedPlayerInput.cpp
@@ -31,6 +31,7 @@ namespace UE
 
 UEnhancedPlayerInput::UEnhancedPlayerInput()
 	: Super()
+	, bIsFlushingInputThisFrame(false)
 {
 	if (UE::Input::EnableDefaultMappingContexts)
 	{
@@ -261,6 +262,9 @@ void UEnhancedPlayerInput::ProcessInputStack(const TArray<UInputComponent*>& Inp
 	static TMap<FKey, bool> KeyDownPrevious;
 	{
 		TRACE_CPUPROFILER_EVENT_SCOPE(EnhPIS_KeyDownPrev);
+
+		const UEnhancedInputDeveloperSettings* Settings = GetDefault<UEnhancedInputDeveloperSettings>();
+
 		KeyDownPrevious.Reset();
 		KeyDownPrevious.Reserve(GetKeyStateMap().Num());
 		for (TPair<FKey, FKeyState>& KeyPair : GetKeyStateMap())
@@ -269,6 +273,19 @@ void UEnhancedPlayerInput::ProcessInputStack(const TArray<UInputComponent*>& Inp
 			// TODO: Can't just use bDownPrevious as paired axis event edges may not fire due to axial deadzoning/missing axis properties. Need to change how this is detected in PlayerInput.cpp.
 			bool bWasDown = KeyState.bDownPrevious || KeyState.EventCounts[IE_Pressed].Num() || KeyState.EventCounts[IE_Repeat].Num();
 			bWasDown |= KeyPair.Key.IsAnalog() && KeyState.RawValue.SizeSquared() != 0;	// Analog inputs should pulse every (non-zero) tick to retain compatibility with UE4.
+
+			// When UPlayerInput::FlushPressedKeys is called any keys that are down will have their RawValue set to 0, and their bDown/bDownPrevious state will be reset to false.
+			// However, their "Value" will not be reset until UPlayerInput::ProcessInputStack. We need to detect if this key was down previously after a flush
+			// so that the Enhanced Input action will correctly fire the triggered values.
+			const bool bKeyWasJustFlushed = 
+				bIsFlushingInputThisFrame && 
+				Settings->bSendTriggeredEventsWhenInputIsFlushed && 
+				!KeyState.Value.IsZero() && 
+				KeyState.RawValue.IsZero() && 
+				!KeyState.bDown;
+
+			bWasDown |= bKeyWasJustFlushed;
+
 			KeyDownPrevious.Emplace(KeyPair.Key, bWasDown);
 		}
 	}
@@ -634,6 +651,14 @@ void UEnhancedPlayerInput::ProcessInputStack(const TArray<UInputComponent*>& Inp
 
 	LastFrameTime = CurrentTime;
 	KeysPressedThisTick.Reset();
+	bIsFlushingInputThisFrame = false;
+}
+
+void UEnhancedPlayerInput::FlushPressedKeys()
+{
+	Super::FlushPressedKeys();
+
+	bIsFlushingInputThisFrame = true;
 }
 
 FInputActionValue UEnhancedPlayerInput::GetActionValue(TObjectPtr<const UInputAction> ForAction) const
diff --git a/Engine/Plugins/EnhancedInput/Source/EnhancedInput/Public/EnhancedInputDeveloperSettings.h b/Engine/Plugins/EnhancedInput/Source/EnhancedInput/Public/EnhancedInputDeveloperSettings.h
--- a/Engine/Plugins/EnhancedInput/Source/EnhancedInput/Public/EnhancedInputDeveloperSettings.h
+++ b/Engine/Plugins/EnhancedInput/Source/EnhancedInput/Public/EnhancedInputDeveloperSettings.h
@@ -79,6 +79,15 @@ public:
 	/** The default player input class that the Enhanced Input world subsystem will use. */
 	UPROPERTY(config, EditAnywhere, NoClear, Category = "Enhanced Input|World Subsystem", meta=(editCondition = "bEnableWorldSubsystem"))
 	TSoftClassPtr<UEnhancedPlayerInput> DefaultWorldInputClass;
+	
+	/**
+	* If true, then any in progress Enhanced Input Actions will fire Cancelled and Triggered events 
+	* when input is flushed (i.e. the viewport has lost focus, or UEnhancedPlayerInput::FlushPressedKeys has been called)
+	* 
+	* If false, then enhanced input actions may not fire their delegates when input is flushed and their key state would be retained.
+	*/
+	UPROPERTY(config, EditAnywhere, Category = "Enhanced Input")
+	uint8 bSendTriggeredEventsWhenInputIsFlushed : 1;
 
 	/**
 	 * If true, then an instance of the User Settings Class will be created on each Enhanced Input subsystem.
diff --git a/Engine/Plugins/EnhancedInput/Source/EnhancedInput/Public/EnhancedPlayerInput.h b/Engine/Plugins/EnhancedInput/Source/EnhancedInput/Public/EnhancedPlayerInput.h
--- a/Engine/Plugins/EnhancedInput/Source/EnhancedInput/Public/EnhancedPlayerInput.h
+++ b/Engine/Plugins/EnhancedInput/Source/EnhancedInput/Public/EnhancedPlayerInput.h
@@ -49,6 +49,10 @@ public:
 
 	UEnhancedPlayerInput();
 
+	//~ Begin UPlayerInput interface
+	virtual void FlushPressedKeys() override;
+	//~ End UPlayerInput interface
+
 	/**
 	* Returns the action instance data for the given input action if there is any. Returns nullptr if the action is not available.
 	*/
@@ -159,6 +163,15 @@ private:
 	 */
 	TArray<FDependentChordTracker> DependentChordActions;
 
+protected:
+
+	/** 
+	* If true, then FlushPressedKeys has been called and the input key state map has been flushed.
+	* 
+	* This will be set to true in UEnhancedPlayerInput::FlushPressedKeys, and reset to false at the end of
+	* UEnhancedPlayerInput::ProcessInputStack
+	*/
+	uint8 bIsFlushingInputThisFrame : 1;
 
 private:
 
diff --git a/Engine/Source/Runtime/Engine/Classes/GameFramework/PlayerInput.h b/Engine/Source/Runtime/Engine/Classes/GameFramework/PlayerInput.h
--- a/Engine/Source/Runtime/Engine/Classes/GameFramework/PlayerInput.h
+++ b/Engine/Source/Runtime/Engine/Classes/GameFramework/PlayerInput.h
@@ -568,7 +568,7 @@ public:
 	//~ End UObject Interface
 
 	/** Flushes the current key state. */
-	void FlushPressedKeys();
+	virtual void FlushPressedKeys();
 
 	/** Flushes the current key state of the keys associated with the action name passed in */
 	void FlushPressedActionBindingKeys(FName ActionName);
