Fix an issue where some key mappings would incorrectly be displayed as overriden in the "showdebug EnhancedInput" command. This was happening because the AddUnique call to the array would always add a new mapping, since it was comparing the instanced arrays of modifiers and triggers.

diff --git a/Engine/Plugins/EnhancedInput/Source/EnhancedInput/Private/EnhancedInputSubsystemInterfaceDebug.cpp b/Engine/Plugins/EnhancedInput/Source/EnhancedInput/Private/EnhancedInputSubsystemInterfaceDebug.cpp
--- a/Engine/Plugins/EnhancedInput/Source/EnhancedInput/Private/EnhancedInputSubsystemInterfaceDebug.cpp
+++ b/Engine/Plugins/EnhancedInput/Source/EnhancedInput/Private/EnhancedInputSubsystemInterfaceDebug.cpp
@@ -136,13 +136,23 @@ void IEnhancedInputSubsystemInterface::ShowMappingContextDebugInfo(UCanvas* Canv
 				}
 			}
 
+			// Some mappings may have been added to the player via the Player Mappable Key system, and thus would not be in any of the 
+			// IMC assets. We can determine that here :) 
 			const TArray<FEnhancedActionKeyMapping>& EnhancedActionMappings = PlayerInput->GetEnhancedActionMappings();
 			for (const FEnhancedActionKeyMapping& EnhancedActionMapping : EnhancedActionMappings)
 			{
+				// Player Mappable keys would have to have an Input Action asset already, so we can filter our search down to that
 				if (TArray<FEnhancedActionKeyMapping>* Mappings = ActionMappings.Find(EnhancedActionMapping.Action))
 				{
-					//Add any mapping that might have been added as player mappable keys.
-					Mappings->AddUnique(EnhancedActionMapping);
+					const bool bMappingExistsAlready = Mappings->ContainsByPredicate([&EnhancedActionMapping](const FEnhancedActionKeyMapping& ExistingMapping)
+						{ 
+							return ExistingMapping.Key == EnhancedActionMapping.Key && ExistingMapping.Action == EnhancedActionMapping.Action;
+						});
+
+					if (!bMappingExistsAlready)
+					{
+						Mappings->Add(EnhancedActionMapping);
+					}					
 				}
 			}
 
