Check the Enhanced Action Mapping modifiers if there are none in the action instance cache for a given mapping. This fixes an issue where debug info about modifiers that are added via IMC assets would not show up when you use the "showdebug enhnacedinput" command.

#rnx

diff --git a/Engine/Plugins/EnhancedInput/Source/EnhancedInput/Private/EnhancedInputSubsystemInterfaceDebug.cpp b/Engine/Plugins/EnhancedInput/Source/EnhancedInput/Private/EnhancedInputSubsystemInterfaceDebug.cpp
--- a/Engine/Plugins/EnhancedInput/Source/EnhancedInput/Private/EnhancedInputSubsystemInterfaceDebug.cpp
+++ b/Engine/Plugins/EnhancedInput/Source/EnhancedInput/Private/EnhancedInputSubsystemInterfaceDebug.cpp
@@ -306,6 +306,27 @@ void IEnhancedInputSubsystemInterface::ShowDebugInfo(UCanvas* Canvas)
 	ShowMappingContextDebugInfo(Canvas, PlayerInput);
 }
 
+
+void IEnhancedInputSubsystemInterface::GetAllRelevantInputModifiersForDebug(const UEnhancedPlayerInput* PlayerInput, const FInputActionInstance* InstanceData, OUT TArray<UInputModifier*>& OutModifiers)
+{
+	OutModifiers.Reset();
+
+	// Start with the modifiers from the instance data, these will be from the Input Action asset
+	OutModifiers.Append(InstanceData->GetModifiers());	
+	
+	// We also need to check the mappings that have been added via the Input Mapping Context
+	for (const FEnhancedActionKeyMapping& Mapping : PlayerInput->GetEnhancedActionMappings())
+	{
+		if (Mapping.Action == InstanceData->GetSourceAction())
+		{
+			for (const TObjectPtr<UInputModifier>& Modifier : Mapping.Modifiers)
+			{
+				OutModifiers.AddUnique(Modifier.Get());
+			}
+		}
+	}
+}
+
 void IEnhancedInputSubsystemInterface::ShowDebugActionModifiers(UCanvas* Canvas, const UInputAction* Action)
 {
 	FDisplayDebugManager& DisplayDebugManager = Canvas->DisplayDebugManager;
@@ -321,7 +342,10 @@ void IEnhancedInputSubsystemInterface::ShowDebugActionModifiers(UCanvas* Canvas,
 	// TODO: Display colored dots to show current sampling locations of valid mappings?
 	// TODO: Invalidate and recalculate textures if sample hash changes? (Rebuild sample data every frame)
 
-	const TArray<UInputModifier*>& Modifiers = ActionData->GetModifiers();
+	static TArray<UInputModifier*> Modifiers;
+
+	IEnhancedInputSubsystemInterface::GetAllRelevantInputModifiersForDebug(PlayerInput, ActionData, OUT Modifiers);
+	
 	if (Modifiers.Num() > 0)
 	{
 		const bool bIs1D = ActionData->GetValue().GetValueType() == EInputActionValueType::Axis1D || ActionData->GetValue().GetValueType() == EInputActionValueType::Boolean;
diff --git a/Engine/Plugins/EnhancedInput/Source/EnhancedInput/Public/EnhancedInputSubsystemInterface.h b/Engine/Plugins/EnhancedInput/Source/EnhancedInput/Public/EnhancedInputSubsystemInterface.h
--- a/Engine/Plugins/EnhancedInput/Source/EnhancedInput/Public/EnhancedInputSubsystemInterface.h
+++ b/Engine/Plugins/EnhancedInput/Source/EnhancedInput/Public/EnhancedInputSubsystemInterface.h
@@ -476,6 +476,9 @@ private:
 	// Debug visualization for enhanced input local player subsystem
 	virtual void ShowDebugInfo(UCanvas* Canvas);
 	
+	// Gather any UInputModifiers on the given Player Input and Instance data that need to be visualized for debugging
+	static void GetAllRelevantInputModifiersForDebug(const UEnhancedPlayerInput* PlayerInput, const FInputActionInstance* InstanceData, OUT TArray<UInputModifier*>& OutModifiers);
+
 	void ShowDebugActionModifiers(UCanvas* Canvas, const UInputAction* Action);
 	static void PurgeDebugVisualizations();
 };
