Add EKeys::AnyKey support to Enhanced Input

diff --git a/Engine/Plugins/EnhancedInput/Source/EnhancedInput/Private/EnhancedPlayerInput.cpp b/Engine/Plugins/EnhancedInput/Source/EnhancedInput/Private/EnhancedPlayerInput.cpp
--- a/Engine/Plugins/EnhancedInput/Source/EnhancedInput/Private/EnhancedPlayerInput.cpp
+++ b/Engine/Plugins/EnhancedInput/Source/EnhancedInput/Private/EnhancedPlayerInput.cpp
@@ -32,6 +32,7 @@ namespace UE
 UEnhancedPlayerInput::UEnhancedPlayerInput()
 	: Super()
 	, bIsFlushingInputThisFrame(false)
+	, CurrentlyInUseAnyKeySubstitute(NAME_None)
 {
 	if (UE::Input::EnableDefaultMappingContexts)
 	{
@@ -264,6 +265,7 @@ void UEnhancedPlayerInput::ProcessInputStack(const TArray<UInputComponent*>& Inp
 		TRACE_CPUPROFILER_EVENT_SCOPE(EnhPIS_KeyDownPrev);
 
 		const UEnhancedInputDeveloperSettings* Settings = GetDefault<UEnhancedInputDeveloperSettings>();
+		bool bWasAnyKeyDownLastFrame = false;
 
 		KeyDownPrevious.Reset();
 		KeyDownPrevious.Reserve(GetKeyStateMap().Num());
@@ -285,9 +287,14 @@ void UEnhancedPlayerInput::ProcessInputStack(const TArray<UInputComponent*>& Inp
 				!KeyState.bDown;
 
 			bWasDown |= bKeyWasJustFlushed;
+			
+			// Keep track of the state of every key so that when we are done iterating we can have a meaningful value for EKeys::AnyKey
+			bWasAnyKeyDownLastFrame |= bWasDown;
 
 			KeyDownPrevious.Emplace(KeyPair.Key, bWasDown);
 		}
+
+		KeyDownPrevious.Emplace(EKeys::AnyKey, bWasAnyKeyDownLastFrame);
 	}
 
 	Super::ProcessInputStack(InputComponentStack, DeltaTime, bGamePaused);
@@ -319,7 +326,48 @@ void UEnhancedPlayerInput::ProcessInputStack(const TArray<UInputComponent*>& Inp
 			continue;
 		}
 
-		FKeyState* KeyState = GetKeyState(Mapping.Key);
+		FKeyState* KeyState = nullptr;
+
+		// If the mapping was to AnyKey, then it won't be in the KeyStateMap and we need to handle it specially. 
+		if (Mapping.Key == EKeys::AnyKey)
+		{
+			// We can just get the first value in the key state map that has been pressed or released, that's what we really care about with this key type
+			for (TPair<FKey, FKeyState>& KeyStatePair : KeyStateMap)
+			{
+				// EKeys::AnyKey will only use non-analog keys. the same as the legacy system.
+				if (!KeyStatePair.Key.IsDigital())
+				{
+					continue;
+				}
+
+				// If we have no Substitute key, then we can simply use the key state of the first available key with pressed events.
+				if (CurrentlyInUseAnyKeySubstitute == NAME_None && KeyStatePair.Value.EventCounts[IE_Pressed].Num() > 0)
+				{
+					KeyState = &KeyStatePair.Value;
+					CurrentlyInUseAnyKeySubstitute = KeyStatePair.Key.GetFName();
+					break;
+				}
+
+				// If we have a substitute key already, then we can just look for the key of this name in the map
+				else if (KeyStatePair.Key.GetFName() == CurrentlyInUseAnyKeySubstitute)
+				{
+					// If the substitute key was just released, then we can reset our currently in 
+					// use substitute key so that it can be replaced by something else on the next key press.					
+					if (KeyStatePair.Value.EventCounts[IE_Released].Num() > 0)
+					{
+						CurrentlyInUseAnyKeySubstitute = NAME_None;
+					}
+
+					KeyState = &KeyStatePair.Value;
+					break;
+				}
+			}
+		}
+		else
+		{
+			KeyState = GetKeyState(Mapping.Key);
+		}
+		
 		FVector RawKeyValue = KeyState ? KeyState->RawValue : FVector::ZeroVector;
 		//UE_CLOG(RawKeyValue.SizeSquared(), LogEnhancedInput, Warning, TEXT("Key %s - state %s"), *Mapping.Key.GetDisplayName().ToString(), *RawKeyValue.ToString());
 
diff --git a/Engine/Plugins/EnhancedInput/Source/EnhancedInput/Public/EnhancedPlayerInput.h b/Engine/Plugins/EnhancedInput/Source/EnhancedInput/Public/EnhancedPlayerInput.h
--- a/Engine/Plugins/EnhancedInput/Source/EnhancedInput/Public/EnhancedPlayerInput.h
+++ b/Engine/Plugins/EnhancedInput/Source/EnhancedInput/Public/EnhancedPlayerInput.h
@@ -173,6 +173,12 @@ protected:
 	*/
 	uint8 bIsFlushingInputThisFrame : 1;
 
+	/**
+	* If there is a key mapping to EKeys::AnyKey, we will keep track of what key was used when we first found a "Pressed"
+	* event. That way we can use the same key when we wait for a "Released" event.
+	*/
+	FName CurrentlyInUseAnyKeySubstitute;
+
 private:
 
 	/** The last time of the last frame that was processed in ProcessPlayerInput */
