Add an advanced setting on the Input Action for the input accumulation behavior. This setting would allow you to change an input action to cumulatively combine input from different Action Key Mappings instead of always accepting the mapping with the highest key value.

diff --git a/Engine/Plugins/EnhancedInput/Source/EnhancedInput/Private/EnhancedPlayerInput.cpp b/Engine/Plugins/EnhancedInput/Source/EnhancedInput/Private/EnhancedPlayerInput.cpp
--- a/Engine/Plugins/EnhancedInput/Source/EnhancedInput/Private/EnhancedPlayerInput.cpp
+++ b/Engine/Plugins/EnhancedInput/Source/EnhancedInput/Private/EnhancedPlayerInput.cpp
@@ -197,6 +197,9 @@ void UEnhancedPlayerInput::ProcessActionMappingEvent(TObjectPtr<const UInputActi
 		// Do this only for no triggers?
 		TriggerStateTracker.SetStateForNoTriggers(ModifiedValue.IsNonZero() ? ETriggerState::Triggered : ETriggerState::None);	
 		bMappingTriggersApplied = Triggers.Num() > 0;
+
+		const EInputActionAccumulationBehavior AccumulationBehavior = ActionData.GetSourceAction()->AccumulationBehavior;
+
 		// Combine values for active events only, selecting the input with the greatest magnitude for each component in each tick.
 		if(ModifiedValue.GetMagnitudeSq())
 		{
@@ -205,10 +208,27 @@ void UEnhancedPlayerInput::ProcessActionMappingEvent(TObjectPtr<const UInputActi
 			FVector Merged = ActionData.Value.Get<FVector>();
 			for (int32 Component = 0; Component < NumComponents; ++Component)
 			{
-				if (FMath::Abs(Modified[Component]) >= FMath::Abs(Merged[Component]))
+				switch (AccumulationBehavior)
+				{
+				// Sometimes you may want to cumulatively merge input. This would allow you to, for example, map WASD to movement and have pressing W and S at the same time
+				// completely cancel out input because "W" is a value of +1.0, and "S" is a value of -1.0
+				case EInputActionAccumulationBehavior::Cumulative:
 				{
-					Merged[Component] = Modified[Component];
+					Merged[Component] += Modified[Component];
+				}										
+				break;
+
+				// By default, we will accept the input with the highest absolute value
+				case EInputActionAccumulationBehavior::TakeHighestAbsoluteValue:
+				default:
+				{
+					if (FMath::Abs(Modified[Component]) >= FMath::Abs(Merged[Component]))
+					{
+						Merged[Component] = Modified[Component];
+					}
 				}
+				break;
+				}			
 			}
 			ActionData.Value = FInputActionValue(ValueType, Merged);
 		}
diff --git a/Engine/Plugins/EnhancedInput/Source/EnhancedInput/Public/InputAction.h b/Engine/Plugins/EnhancedInput/Source/EnhancedInput/Public/InputAction.h
--- a/Engine/Plugins/EnhancedInput/Source/EnhancedInput/Public/InputAction.h
+++ b/Engine/Plugins/EnhancedInput/Source/EnhancedInput/Public/InputAction.h
@@ -14,7 +14,41 @@ class UPlayerMappableKeySettings;
 
 enum class ETriggerEventInternal : uint8;
 
-// Input action definition. These are instanced per player (via FInputActionInstance)
+/** 
+* This is an advanced setting that allows you to change how the value of an Input Action is calculated when there are 
+* multiple mappings to the same Input Action. The default behavior is to accept highest absolute value.
+*/
+UENUM()
+enum class EInputActionAccumulationBehavior : uint8
+{
+	/** 
+	* Take the value from the mapping with the highest Absolute Value.
+	* 
+	* For example, given a value of -0.3 and 0.5, the input action's value would be 0.5. 
+	*/
+	TakeHighestAbsoluteValue,
+
+	/** 
+	* Cumulatively adds the key values for each mapping.
+	* 
+	* For example, a value of -0.7 and +0.75 on the same input action would result in a value of 0.05.
+	* 
+	* A practical example of when to use this would be for something like WASD movement, if you want pressing W and S to cancel each other out.
+	*/
+	Cumulative,
+};
+
+/**
+* An Input Action is a logical representation of something the user can do, such as "Jump" or "Crouch".
+* These are what your gameplay code binds to in order to listen for input state changes. For most scenarios 
+* your gameplay code should be listening for the "Triggered" event on an input action. This will allow
+* for the most scalable and customizable input configuration because you can add different triggers 
+* for each key mapping in the Input Mapping Context. 
+* 
+* They are the conceptual equivalent to "Action" and "Axis" mapping names from the Legacy Input System.
+* 
+* Note: These are instanced per player (via FInputActionInstance)
+*/
 UCLASS(BlueprintType)
 class ENHANCEDINPUT_API UInputAction : public UDataAsset
 {
@@ -66,6 +100,17 @@ public:
 	// The type that this action returns from a GetActionValue query or action event
 	UPROPERTY(EditAnywhere, BlueprintReadOnly, Category = Action, AssetRegistrySearchable)
 	EInputActionValueType ValueType = EInputActionValueType::Boolean;
+
+	/**
+	* This defines how the value of this input action will be calcuated in the case that there are multiple key mappings to the same input action.
+	* 
+	* When TakeHighestAbsoluteValue is selected, then the key mapping with the highest absolutle value will be utilized. (Default)
+	* When Cumulative is selected, then each key mapping will be added together to get the key value. 
+	* 
+	* @see UEnhancedPlayerInput::ProcessActionMappingEvent, where this property is read from. 
+	*/
+	UPROPERTY(EditAnywhere, BlueprintReadOnly, Category = Action, AdvancedDisplay)
+	EInputActionAccumulationBehavior AccumulationBehavior = EInputActionAccumulationBehavior::TakeHighestAbsoluteValue;
 	
 	/**
 	* Trigger qualifiers. If any trigger qualifiers exist the action will not trigger unless:
