Deprecate the FPlayerMappableKeySlot and all its related functoins for 5.3, along with all the old player mapped key functions in favor of UEnhancedInputUserSettings

diff --git a/Engine/Plugins/EnhancedInput/Source/EnhancedInput/Private/EnhancedInputSubsystemInterface.cpp b/Engine/Plugins/EnhancedInput/Source/EnhancedInput/Private/EnhancedInputSubsystemInterface.cpp
--- a/Engine/Plugins/EnhancedInput/Source/EnhancedInput/Private/EnhancedInputSubsystemInterface.cpp
+++ b/Engine/Plugins/EnhancedInput/Source/EnhancedInput/Private/EnhancedInputSubsystemInterface.cpp
@@ -523,6 +523,8 @@ TArray<FEnhancedActionKeyMapping> IEnhancedInputSubsystemInterface::GetAllPlayer
 	return PlayerMappableMappings;
 }
 
+PRAGMA_DISABLE_DEPRECATION_WARNINGS
+
 int32 IEnhancedInputSubsystemInterface::AddPlayerMappedKey(const FName MappingName, const FKey NewKey, const FModifyContextOptions& Options /*= FModifyContextOptions()*/)
 {
 	return AddPlayerMappedKeyInSlot(MappingName, NewKey, FPlayerMappableKeySlot::FirstKeySlot, Options);
@@ -535,24 +537,25 @@ int32 IEnhancedInputSubsystemInterface::K2_AddPlayerMappedKeyInSlot(const FName
 
 int32 IEnhancedInputSubsystemInterface::AddPlayerMappedKeyInSlot(const FName MappingName, const FKey NewKey, const FPlayerMappableKeySlot& KeySlot /*= FPlayerMappableKeySlot::FirstKeySlot*/, const FModifyContextOptions& Options /*= FModifyContextOptions()*/)
 {
-	int32 NumMappingsApplied = 0;
-	if (MappingName != NAME_None)
+	if (UEnhancedInputUserSettings* Settings = GetUserSettings())
 	{
-		if (UEnhancedPlayerInput* const PlayerInput = GetPlayerInput())
+		FMapPlayerKeyArgs Args = {};
+		Args.MappingName = MappingName;
+		Args.NewKey = NewKey;
+		Args.Slot = static_cast<EPlayerMappableKeySlot>(KeySlot.GetSlotNumber());
+
+		FGameplayTagContainer FailureReason;
+		Settings->MapPlayerKey(Args, FailureReason);
+
+		if (!FailureReason.IsEmpty())
 		{
-			TMap<FPlayerMappableKeySlot, FKey>& PlayerMappedKeySlots = PlayerMappedSettings.FindOrAdd(MappingName);
-			FKey& KeyInSlot = PlayerMappedKeySlots.FindOrAdd(KeySlot);
-			KeyInSlot = NewKey;
-			++NumMappingsApplied;
+			UE_LOG(LogEnhancedInput, Warning, TEXT("IEnhancedInputSubsystemInterface::AddPlayerMappedKeyInSlot Failed! Reasoning: %s"), *FailureReason.ToString());
+			return 0;
 		}
-
-		RequestRebuildControlMappings(Options);
 	}
-	else
-	{
-		UE_LOG(LogEnhancedInput, Warning, TEXT("Attempted to AddPlayerMappedKeyInSlot with an invalid MappingName! Mapping has not been applied."));
-	}
-	return NumMappingsApplied;
+
+	RequestRebuildControlMappings(Options);
+	return 1;
 }
 
 int32 IEnhancedInputSubsystemInterface::RemovePlayerMappedKey(const FName MappingName, const FModifyContextOptions& Options /*= FModifyContextOptions()*/)
@@ -567,38 +570,55 @@ int32 IEnhancedInputSubsystemInterface::K2_RemovePlayerMappedKeyInSlot(const FNa
 
 int32 IEnhancedInputSubsystemInterface::RemovePlayerMappedKeyInSlot(const FName MappingName, const FPlayerMappableKeySlot& KeySlot /*= FPlayerMappableKeySlot::FirstKeySlot*/, const FModifyContextOptions& Options /*= FModifyContextOptions()*/)
 {
-	int32 NumMappingsRemoved = 0;
-	if (UEnhancedPlayerInput* const PlayerInput = GetPlayerInput())
+	if (UEnhancedInputUserSettings* Settings = GetUserSettings())
 	{
-		if (TMap<FPlayerMappableKeySlot, FKey>* PlayerMappedKeySlots = PlayerMappedSettings.Find(MappingName))
+		FMapPlayerKeyArgs Args = {};
+		Args.MappingName = MappingName;
+		Args.Slot = static_cast<EPlayerMappableKeySlot>(KeySlot.GetSlotNumber());
+		
+		FGameplayTagContainer FailureReason;
+		Settings->UnMapPlayerKey(Args, FailureReason);
+
+		if (!FailureReason.IsEmpty())
 		{
-			NumMappingsRemoved = PlayerMappedKeySlots->Remove(KeySlot);
+			UE_LOG(LogEnhancedInput, Warning, TEXT("IEnhancedInputSubsystemInterface::RemovePlayerMappedKeyInSlot Failed! Reasoning: %s"), *FailureReason.ToString());
+			return 0;
 		}
 	}
-
+	
 	RequestRebuildControlMappings(Options);
 
-	return NumMappingsRemoved;
+	return 1;
 }
 
 int32 IEnhancedInputSubsystemInterface::RemoveAllPlayerMappedKeysForMapping(const FName MappingName, const FModifyContextOptions& Options /*= FModifyContextOptions()*/)
 {
-	int32 NumMappingsRemoved = 0;
-	if (UEnhancedPlayerInput* const PlayerInput = GetPlayerInput())
+	if (UEnhancedInputUserSettings* Settings = GetUserSettings())
 	{
-		NumMappingsRemoved = PlayerMappedSettings.Remove(MappingName);
+		FMapPlayerKeyArgs Args = {};
+		Args.MappingName = MappingName;
+
+		FGameplayTagContainer FailureReason;
+		Settings->ResetAllPlayerKeysInRow(Args, FailureReason);
+		
+		if (!FailureReason.IsEmpty())
+		{
+			UE_LOG(LogEnhancedInput, Warning, TEXT("IEnhancedInputSubsystemInterface::RemovePlayerMappedKeyInSlot Failed! Reasoning: %s"), *FailureReason.ToString());
+			return 0;
+		}
 	}
 
 	RequestRebuildControlMappings(Options);
 
-	return NumMappingsRemoved;
+	return 1;
 }
 
 void IEnhancedInputSubsystemInterface::RemoveAllPlayerMappedKeys(const FModifyContextOptions& Options)
 {
-	if (UEnhancedPlayerInput* const PlayerInput = GetPlayerInput())
+	if (UEnhancedInputUserSettings* Settings = GetUserSettings())
 	{
-		PlayerMappedSettings.Empty();
+		FGameplayTagContainer FailureReason;
+		Settings->ResetKeyProfileToDefault(Settings->GetCurrentKeyProfileIdentifier(), FailureReason);
 	}
 
 	RequestRebuildControlMappings(Options);
@@ -616,24 +636,42 @@ FKey IEnhancedInputSubsystemInterface::K2_GetPlayerMappedKeyInSlot(const FName M
 
 FKey IEnhancedInputSubsystemInterface::GetPlayerMappedKeyInSlot(const FName MappingName, const FPlayerMappableKeySlot& KeySlot /*= FPlayerMappableKeySlot()*/) const
 {
-	if (const TMap<FPlayerMappableKeySlot, FKey>* PlayerMappedKeySlots = PlayerMappedSettings.Find(MappingName))
+	if (const UEnhancedInputUserSettings* Settings = GetUserSettings())
 	{
-		if (const FKey* PlayerMappedKey = PlayerMappedKeySlots->Find(KeySlot))
+		if (const UEnhancedPlayerMappableKeyProfile* KeyProfile = Settings->GetCurrentKeyProfile())
 		{
-			return *PlayerMappedKey;
+			FPlayerMappableKeyQueryOptions Opts = {};
+			Opts.MappingName = MappingName;
+			Opts.SlotToMatch = static_cast<EPlayerMappableKeySlot>(KeySlot.GetSlotNumber());
+
+			TArray<FKey> Keys;
+			KeyProfile->QueryPlayerMappedKeys(Opts, OUT Keys);
+			
+			if (!Keys.IsEmpty())
+			{
+				return Keys[0];
+			}
 		}
 	}
+	
 	return EKeys::Invalid;
 }
 
 TArray<FKey> IEnhancedInputSubsystemInterface::GetAllPlayerMappedKeys(const FName MappingName) const
 {
 	TArray<FKey> PlayerMappedKeys;
-	if (const TMap<FPlayerMappableKeySlot, FKey>* PlayerMappedKeySlots = PlayerMappedSettings.Find(MappingName))
+	
+	if (const UEnhancedInputUserSettings* Settings = GetUserSettings())
 	{
-		PlayerMappedKeySlots->GenerateValueArray(PlayerMappedKeys);
-		return PlayerMappedKeys;
+		if (const UEnhancedPlayerMappableKeyProfile* KeyProfile = Settings->GetCurrentKeyProfile())
+		{
+			FPlayerMappableKeyQueryOptions Opts = {};
+			Opts.MappingName = MappingName;
+			
+			KeyProfile->QueryPlayerMappedKeys(Opts, OUT PlayerMappedKeys);
+		}
 	}
+	
 	return PlayerMappedKeys;
 }
 
@@ -669,6 +707,8 @@ void IEnhancedInputSubsystemInterface::RemovePlayerMappableConfig(const UPlayerM
 	}
 }
 
+PRAGMA_ENABLE_DEPRECATION_WARNINGS
+
 template<typename T>
 void DeepCopyPtrArray(const TArray<T*>& From, TArray<T*>& To)
 {
@@ -869,10 +909,6 @@ void IEnhancedInputSubsystemInterface::RebuildControlMappings()
 			{
 				PlayerKeyProfile->GetPlayerMappedKeysForRebuildControlMappings(Mapping, MappedKeysToActionName);
 			}
-			else if (const TMap<FPlayerMappableKeySlot, FKey>* PlayerMappedKeySlots = PlayerMappedSettings.Find(Mapping.GetMappingName()))
-			{
-				PlayerMappedKeySlots->GenerateValueArray(MappedKeysToActionName);
-			}
 
 			// True if there were any player mapped keys to this mapping and we are using those instead.
 			const bool bIsPlayerMapping = !MappedKeysToActionName.IsEmpty();
diff --git a/Engine/Plugins/EnhancedInput/Source/EnhancedInput/Private/PlayerMappableKeySlot.cpp b/Engine/Plugins/EnhancedInput/Source/EnhancedInput/Private/PlayerMappableKeySlot.cpp
--- a/Engine/Plugins/EnhancedInput/Source/EnhancedInput/Private/PlayerMappableKeySlot.cpp
+++ b/Engine/Plugins/EnhancedInput/Source/EnhancedInput/Private/PlayerMappableKeySlot.cpp
@@ -2,6 +2,8 @@
 
 #include "PlayerMappableKeySlot.h"
 
+PRAGMA_DISABLE_DEPRECATION_WARNINGS
+
 FPlayerMappableKeySlot FPlayerMappableKeySlot::FirstKeySlot = FPlayerMappableKeySlot(0);
 FPlayerMappableKeySlot FPlayerMappableKeySlot::SecondKeySlot = FPlayerMappableKeySlot(1);
 FPlayerMappableKeySlot FPlayerMappableKeySlot::ThirdKeySlot = FPlayerMappableKeySlot(2);
@@ -9,7 +11,6 @@ FPlayerMappableKeySlot FPlayerMappableKeySlot::FourthKeySlot = FPlayerMappableKe
 
 FPlayerMappableKeySlot::FPlayerMappableKeySlot() { }
 FPlayerMappableKeySlot::FPlayerMappableKeySlot(const int32 InSlotNumber) : SlotNumber(InSlotNumber) { }
-FPlayerMappableKeySlot::~FPlayerMappableKeySlot() { }
 
 int32 FPlayerMappableKeySlot::GetSlotNumber() const
 {
@@ -20,3 +21,5 @@ uint32 GetTypeHash(const FPlayerMappableKeySlot& Ref)
 {
 	return GetTypeHash(Ref.GetSlotNumber());
 }
+
+PRAGMA_ENABLE_DEPRECATION_WARNINGS
\ No newline at end of file
diff --git a/Engine/Plugins/EnhancedInput/Source/EnhancedInput/Public/EnhancedInputSubsystemInterface.h b/Engine/Plugins/EnhancedInput/Source/EnhancedInput/Public/EnhancedInputSubsystemInterface.h
--- a/Engine/Plugins/EnhancedInput/Source/EnhancedInput/Public/EnhancedInputSubsystemInterface.h
+++ b/Engine/Plugins/EnhancedInput/Source/EnhancedInput/Public/EnhancedInputSubsystemInterface.h
@@ -280,13 +280,18 @@ public:
 	UE_DEPRECATED(5.2, "AddPlayerMappedKey has been deprecated, please use AddPlayerMappedKeyInSlot instead.")
 	virtual int32 AddPlayerMappedKey(const FName MappingName, const FKey NewKey, const FModifyContextOptions& Options = FModifyContextOptions());
 
+	PRAGMA_DISABLE_DEPRECATION_WARNINGS
 	/**
 	 * Emplace or replace any currently applied key in KeySlot for mapping of MappingName.
 	 * Requests a rebuild of the player mappings.
 	 *
 	 * @return The number of mappings that have been replaced
 	 */
-	UFUNCTION(BlueprintCallable, BlueprintCosmetic, Category = "Input|PlayerMappable", meta = (DisplayName = "Add Player Mapped Key In Slot", AutoCreateRefTerm = "KeySlot, Options"))
+	UE_DEPRECATED(5.3, "K2_AddPlayerMappedKeyInSlot has been deprecated, please use UEnhancedInputUserSettings instead.")
+	UFUNCTION(BlueprintCallable, BlueprintCosmetic, Category = "Input|PlayerMappable", meta = (
+		DeprecatedFunction,
+		DeprecationMessage="AddPlayerMappedKeyInSlot has been deprecated, please use UEnhancedInputUserSettings instead.",
+		DisplayName = "Add Player Mapped Key In Slot", AutoCreateRefTerm = "KeySlot, Options"))
 	virtual int32 K2_AddPlayerMappedKeyInSlot(const FName MappingName, const FKey NewKey, const FPlayerMappableKeySlot& KeySlot = FPlayerMappableKeySlot(), const FModifyContextOptions& Options = FModifyContextOptions());
 
 	/**
@@ -295,6 +300,7 @@ public:
 	 *
 	 * @return The number of mappings that have been replaced
 	 */
+	UE_DEPRECATED(5.3, "AddPlayerMappedKeyInSlot has been deprecated, please use UEnhancedInputUserSettings instead.")
 	virtual int32 AddPlayerMappedKeyInSlot(const FName MappingName, const FKey NewKey, const FPlayerMappableKeySlot& KeySlot = FPlayerMappableKeySlot::FirstKeySlot, const FModifyContextOptions& Options = FModifyContextOptions());
 
 	/**
@@ -312,7 +318,9 @@ public:
 	 *
 	 * @return The number of mappings that have been removed
 	 */
-	UFUNCTION(BlueprintCallable, BlueprintCosmetic, Category = "Input|PlayerMappable", meta = (DisplayName="Remove Player Mapped Key In Slot", AutoCreateRefTerm = "KeySlot, Options"))
+	UFUNCTION(BlueprintCallable, BlueprintCosmetic, Category = "Input|PlayerMappable", meta = (DeprecatedFunction,
+		DeprecationMessage="K2_RemovePlayerMappedKeyInSlot has been deprecated, please use UEnhancedInputUserSettings instead.",
+		DisplayName="Remove Player Mapped Key In Slot", AutoCreateRefTerm = "KeySlot, Options"))
 	virtual int32 K2_RemovePlayerMappedKeyInSlot(const FName MappingName, const FPlayerMappableKeySlot& KeySlot = FPlayerMappableKeySlot(), const FModifyContextOptions& Options = FModifyContextOptions());
 
 	/**
@@ -321,6 +329,7 @@ public:
 	 *
 	 * @return The number of mappings that have been removed
 	 */
+	UE_DEPRECATED(5.3, "RemovePlayerMappedKeyInSlot has been deprecated, please use UEnhancedInputUserSettings instead.")
 	virtual int32 RemovePlayerMappedKeyInSlot(const FName MappingName, const FPlayerMappableKeySlot& KeySlot = FPlayerMappableKeySlot::FirstKeySlot, const FModifyContextOptions& Options = FModifyContextOptions());
 
 	/**
@@ -329,7 +338,10 @@ public:
 	 *
 	 * @return The number of mappings that have been removed
 	 */
-	UFUNCTION(BlueprintCallable, BlueprintCosmetic, Category = "Input|PlayerMappable", meta = (AutoCreateRefTerm = "KeySlot, Options"))
+	UE_DEPRECATED(5.3, "RemovePlayerMappedKeyInSlot has been deprecated, please use UEnhancedInputUserSettings instead.")
+	UFUNCTION(BlueprintCallable, BlueprintCosmetic, Category = "Input|PlayerMappable", meta = (DeprecatedFunction,
+		DeprecationMessage="RemoveAllPlayerMappedKeysForMapping has been deprecated, please use UEnhancedInputUserSettings instead.",
+		AutoCreateRefTerm = "KeySlot, Options"))
 	virtual int32 RemoveAllPlayerMappedKeysForMapping(const FName MappingName, const FModifyContextOptions& Options = FModifyContextOptions());
 
 	/**
@@ -347,7 +359,11 @@ public:
 	 *
 	 * @param MappingName	The FName of the mapped key that would have been set with the AddPlayerMappedKey function.
 	 */
-	UFUNCTION(BlueprintCallable, BlueprintCosmetic, Category = "Input|PlayerMappable", meta = (DisplayName="Get Player Mapped Key In Slot",AutoCreateRefTerm = "KeySlot"))
+	UE_DEPRECATED(5.3, "RemovePlayerMappedKeyInSlot has been deprecated, please use UEnhancedInputUserSettings instead.")
+	UFUNCTION(BlueprintCallable, BlueprintCosmetic, Category = "Input|PlayerMappable", meta = (
+		DeprecatedFunction,
+		DeprecationMessage="K2_GetPlayerMappedKeyInSlot has been deprecated, please use UEnhancedInputUserSettings instead.",
+		DisplayName="Get Player Mapped Key In Slot",AutoCreateRefTerm = "KeySlot"))
 	virtual FKey K2_GetPlayerMappedKeyInSlot(const FName MappingName, const FPlayerMappableKeySlot& KeySlot = FPlayerMappableKeySlot()) const;
 
 	/**
@@ -356,6 +372,7 @@ public:
 	 *
 	 * @param MappingName	The FName of the mapped key that would have been set with the AddPlayerMappedKey function.
 	 */
+	UE_DEPRECATED(5.3, "RemovePlayerMappedKeyInSlot has been deprecated, please use UEnhancedInputUserSettings instead.")
 	virtual FKey GetPlayerMappedKeyInSlot(const FName MappingName, const FPlayerMappableKeySlot& KeySlot = FPlayerMappableKeySlot::FirstKeySlot) const;
 
 	/**
@@ -364,7 +381,8 @@ public:
 	 *
 	 * @param MappingName	The FName of the mapped key that would have been set with the AddPlayerMappedKey function.
 	 */
-	UFUNCTION(BlueprintCallable, BlueprintCosmetic, Category = "Input|PlayerMappable")
+	UE_DEPRECATED(5.3, "GetAllPlayerMappedKeys has been deprecated, please use UEnhancedInputUserSettings instead.")
+	UFUNCTION(BlueprintCallable, BlueprintCosmetic, Category = "Input|PlayerMappable", meta=(DeprecatedFunction, DeprecationMessage="GetAllPlayerMappedKeys has been deprecated, please use UEnhancedInputUserSettings instead."))
 	virtual TArray<FKey> GetAllPlayerMappedKeys(const FName MappingName) const;
 	
 	/**
@@ -372,17 +390,22 @@ public:
 	 * Requests a rebuild of the player mappings. 
 	 *
 	 */
-	UFUNCTION(BlueprintCallable, BlueprintCosmetic, Category = "Input|PlayerMappable", meta=(AutoCreateRefTerm = "Options"))
+	UE_DEPRECATED(5.3, "RemoveAllPlayerMappedKeys has been deprecated, please use UEnhancedInputUserSettings instead.")
+	UFUNCTION(BlueprintCallable, BlueprintCosmetic, Category = "Input|PlayerMappable", meta=(AutoCreateRefTerm = "Options", DeprecatedFunction, DeprecationMessage="RemoveAllPlayerMappedKeys has been deprecated, please use UEnhancedInputUserSettings instead."))
 	virtual void RemoveAllPlayerMappedKeys(const FModifyContextOptions& Options = FModifyContextOptions());
 	
 	/** Adds all the input mapping contexts inside of this mappable config. */
-	UFUNCTION(BlueprintCallable, BlueprintCosmetic, Category = "Input|PlayerMappable", meta=(AutoCreateRefTerm = "Options"))
+	UE_DEPRECATED(5.3, "RemoveAllPlayerMappedKeys has been deprecated, please use UEnhancedInputUserSettings instead.")
+	UFUNCTION(BlueprintCallable, BlueprintCosmetic, Category = "Input|PlayerMappable", meta=(AutoCreateRefTerm = "Options", DeprecatedFunction, DeprecationMessage="AddPlayerMappableConfig has been deprecated, please use UEnhancedInputUserSettings instead."))
 	virtual void AddPlayerMappableConfig(const UPlayerMappableInputConfig* Config, const FModifyContextOptions& Options = FModifyContextOptions());
 
 	/** Removes all the input mapping contexts inside of this mappable config. */
-	UFUNCTION(BlueprintCallable, BlueprintCosmetic, Category = "Input|PlayerMappable", meta=(AutoCreateRefTerm = "Options"))
+	UE_DEPRECATED(5.3, "RemoveAllPlayerMappedKeys has been deprecated, please use UEnhancedInputUserSettings instead.")
+	UFUNCTION(BlueprintCallable, BlueprintCosmetic, Category = "Input|PlayerMappable", meta=(AutoCreateRefTerm = "Options", DeprecatedFunction, DeprecationMessage="RemovePlayerMappableConfig has been deprecated, please use UEnhancedInputUserSettings instead."))
 	virtual void RemovePlayerMappableConfig(const UPlayerMappableInputConfig* Config, const FModifyContextOptions& Options = FModifyContextOptions());
-
+	
+	PRAGMA_ENABLE_DEPRECATION_WARNINGS
+	
 protected:
 
 	/**
@@ -437,9 +460,6 @@ private:
 	TMap<TWeakObjectPtr<const UInputAction>, FInputActionValue> ForcedActions;
 	TMap<FKey, FInputActionValue> ForcedKeys;
 
-	/** A map of any player mapped keys to the keys that they should redirect to instead */
-	TMap<FName, TMap<FPlayerMappableKeySlot, FKey>> PlayerMappedSettings;
-
 	/** A map of inputs that should be injected every frame. This inputs will be injected when ForcedInput is ticked. */
 	TMap<TObjectPtr<const UInputAction>, FInjectedInput> ContinuouslyInjectedInputs;
 
diff --git a/Engine/Plugins/EnhancedInput/Source/EnhancedInput/Public/PlayerMappableKeySlot.h b/Engine/Plugins/EnhancedInput/Source/EnhancedInput/Public/PlayerMappableKeySlot.h
--- a/Engine/Plugins/EnhancedInput/Source/EnhancedInput/Public/PlayerMappableKeySlot.h
+++ b/Engine/Plugins/EnhancedInput/Source/EnhancedInput/Public/PlayerMappableKeySlot.h
@@ -8,6 +8,9 @@
 * Explicitly identifies the slot for a player mappable key
 * Experimental: Do not count on long term support for this structure.
 */
+
+struct UE_DEPRECATED(5.3, "FPlayerMappableKeySlot has been deprecated. Please use EPlayerMappableKeySlot instead.") FPlayerMappableKeySlot;
+
 USTRUCT(BlueprintType, DisplayName="Player Mappable Key Slot (Experimental)", meta = (ScriptName = "PlayerMappableKeySlotData"))
 struct ENHANCEDINPUT_API FPlayerMappableKeySlot
 {
@@ -17,10 +20,17 @@ public:
 
 	FPlayerMappableKeySlot();
 	FPlayerMappableKeySlot(int32 InSlotNumber);
-	virtual ~FPlayerMappableKeySlot();
 
-	virtual int32 GetSlotNumber() const;
+	PRAGMA_DISABLE_DEPRECATION_WARNINGS
+	FPlayerMappableKeySlot(FPlayerMappableKeySlot&&) = default;
+	FPlayerMappableKeySlot(const FPlayerMappableKeySlot& Other) = default;
+	FPlayerMappableKeySlot& operator=(const FPlayerMappableKeySlot& Other) = default;
+	FPlayerMappableKeySlot& operator=(FPlayerMappableKeySlot&&) = default;
+	~FPlayerMappableKeySlot() = default;
+	PRAGMA_ENABLE_DEPRECATION_WARNINGS
 
+	int32 GetSlotNumber() const;
+	
 	UPROPERTY(EditAnywhere, BlueprintReadWrite, Category = "Input")
 	int32 SlotNumber = 0;
 
@@ -36,4 +46,8 @@ public:
 
 };
 
+PRAGMA_DISABLE_DEPRECATION_WARNINGS
+
 ENHANCEDINPUT_API uint32 GetTypeHash(const FPlayerMappableKeySlot& Ref);
+
+PRAGMA_ENABLE_DEPRECATION_WARNINGS
\ No newline at end of file
