Add the "Config" metadata to the bShouldOnlyTriggerLastActionInChord and bEnableDefaultMappingContexts so that they can be set in a config file as well as with a console variable.

#rnx

diff --git a/Engine/Plugins/EnhancedInput/Source/EnhancedInput/Public/EnhancedInputDeveloperSettings.h b/Engine/Plugins/EnhancedInput/Source/EnhancedInput/Public/EnhancedInputDeveloperSettings.h
--- a/Engine/Plugins/EnhancedInput/Source/EnhancedInput/Public/EnhancedInputDeveloperSettings.h
+++ b/Engine/Plugins/EnhancedInput/Source/EnhancedInput/Public/EnhancedInputDeveloperSettings.h
@@ -98,7 +98,7 @@ public:
 	uint8 bEnableUserSettings : 1;
 
 	/** If true, then the DefaultMappingContexts will be applied to all Enhanced Input Subsystems. */
-	UPROPERTY(EditAnywhere, Category = "Enhanced Input", meta = (ConsoleVariable = "EnhancedInput.EnableDefaultMappingContexts"))
+	UPROPERTY(config, EditAnywhere, Category = "Enhanced Input", meta = (ConsoleVariable = "EnhancedInput.EnableDefaultMappingContexts"))
 	uint8 bEnableDefaultMappingContexts : 1;
 
 	/**
@@ -107,7 +107,7 @@ public:
 	 * 
 	 * Default value is true.
 	 */
-	UPROPERTY(EditAnywhere, Category = "Enhanced Input", meta=(ConsoleVariable="EnhancedInput.OnlyTriggerLastActionInChord"))
+	UPROPERTY(config, EditAnywhere, Category = "Enhanced Input", meta=(ConsoleVariable="EnhancedInput.OnlyTriggerLastActionInChord"))
 	uint8 bShouldOnlyTriggerLastActionInChord : 1;
 
 	/**
