Remove a LogTemp

#rnx

diff --git a/Engine/Plugins/EnhancedInput/Source/EnhancedInput/Private/EnhancedInputSubsystemInterfaceDebug.cpp b/Engine/Plugins/EnhancedInput/Source/EnhancedInput/Private/EnhancedInputSubsystemInterfaceDebug.cpp
--- a/Engine/Plugins/EnhancedInput/Source/EnhancedInput/Private/EnhancedInputSubsystemInterfaceDebug.cpp
+++ b/Engine/Plugins/EnhancedInput/Source/EnhancedInput/Private/EnhancedInputSubsystemInterfaceDebug.cpp
@@ -473,8 +473,7 @@ void IEnhancedInputSubsystemInterface::ShowDebugActionModifiers(UCanvas* Canvas,
 
 					}
 				}
-
-				UE_LOG(LogTemp, Warning, TEXT("Building visualization for %s:%s (of %d)"), *PlayerInput->GetName(), *Modifier->GetName(), CachedModifierVisualizations.Num());
+				
 				Visualization = CachedModifierVisualizations.Add((uint64)Modifier, FVisualizationTexture(BuildVisualizationTexture(Modifier, SampleData))).Texture;
 				bBuiltAtLeastOneModifier = true;
 			}
