Only draw debug modifiers if we are in the editor. The drawing of these debug textures can cause some bad perf hitching while using the debug command if in a cooked game client, and they would only ever draw in White boxes anyway. There is no point to attempting to draw this outside the editor.

#rnx

diff --git a/Engine/Plugins/EnhancedInput/Source/EnhancedInput/Private/EnhancedInputSubsystemInterfaceDebug.cpp b/Engine/Plugins/EnhancedInput/Source/EnhancedInput/Private/EnhancedInputSubsystemInterfaceDebug.cpp
--- a/Engine/Plugins/EnhancedInput/Source/EnhancedInput/Private/EnhancedInputSubsystemInterfaceDebug.cpp
+++ b/Engine/Plugins/EnhancedInput/Source/EnhancedInput/Private/EnhancedInputSubsystemInterfaceDebug.cpp
@@ -255,7 +255,11 @@ void IEnhancedInputSubsystemInterface::ShowMappingContextDebugInfo(UCanvas* Canv
 				}
 				DisplayDebugManager.DrawString(ActionStr);
 
+				// The drawing of these debug textures only works in the editor and causes some serious perf hitches 
+				// when you try to call it in a cooked game build.
+#if WITH_EDITOR
 				ShowDebugActionModifiers(Canvas, Action);
+#endif	// WITH_EDITOR
 
 				for (const FDisplayLine& DisplayLine : MappingDisplayLine)
 				{
