Add a null check to IMC's on the debug drawing screen. It is possible that if you are not using actual data assets for Enhanced Input and instead are spawning them dynamically at runtime that the IMC could be garabage collected when the debug screen is displayed. Add this check here just in case that happens.

#rnx

diff --git a/Engine/Plugins/EnhancedInput/Source/EnhancedInput/Private/EnhancedInputSubsystemInterfaceDebug.cpp b/Engine/Plugins/EnhancedInput/Source/EnhancedInput/Private/EnhancedInputSubsystemInterfaceDebug.cpp
--- a/Engine/Plugins/EnhancedInput/Source/EnhancedInput/Private/EnhancedInputSubsystemInterfaceDebug.cpp
+++ b/Engine/Plugins/EnhancedInput/Source/EnhancedInput/Private/EnhancedInputSubsystemInterfaceDebug.cpp
@@ -107,6 +107,16 @@ void IEnhancedInputSubsystemInterface::ShowMappingContextDebugInfo(UCanvas* Canv
 		{
 			const UInputMappingContext* AppliedContext = ContextPair.Key.Get();
 
+			// If you for some reason choose to not use IMC assets and instead dynamically spawn them at runtime, they could be
+			// garbage collected. This is NOT a recommended approach to using Enhanced Input, but we will put this error here
+			// just in case.
+			if (!AppliedContext)
+			{
+				DisplayDebugManager.SetDrawColor(FColor::Red);
+				DisplayDebugManager.DrawString(TEXT("Null Input Mapping Context!"));
+				continue;
+			}
+
 			DisplayDebugManager.SetDrawColor(FColor::Yellow);
 			
 			// If this context was redirected via platform settings, then add some debug info about it here
