Allow for the settings delegate to be deferred to the next tick when UnMapping a key, just like it can be with MapPlayerKey

#rnx

diff --git a/Engine/Plugins/EnhancedInput/Source/EnhancedInput/Private/UserSettings/EnhancedInputUserSettings.cpp b/Engine/Plugins/EnhancedInput/Source/EnhancedInput/Private/UserSettings/EnhancedInputUserSettings.cpp
--- a/Engine/Plugins/EnhancedInput/Source/EnhancedInput/Private/UserSettings/EnhancedInputUserSettings.cpp
+++ b/Engine/Plugins/EnhancedInput/Source/EnhancedInput/Private/UserSettings/EnhancedInputUserSettings.cpp
@@ -988,7 +988,27 @@ void UEnhancedInputUserSettings::UnMapPlayerKey(const FMapPlayerKeyArgs& InArgs,
 		if (FoundMapping->IsDirty())
 		{
 			OnKeyMappingUpdated(FoundMapping, InArgs, true);
-			OnSettingsChanged.Broadcast(this);
+
+			if (InArgs.bDeferOnSettingsChangedBroadcast  && !DeferredSettingsChangedTimerHandle.IsValid())
+			{
+				if (UWorld* World = GetWorld())
+				{
+					TWeakObjectPtr<UEnhancedInputUserSettings> WeakThis = this;
+					DeferredSettingsChangedTimerHandle = World->GetTimerManager().SetTimerForNextTick(FTimerDelegate::CreateWeakLambda(this, 
+						[WeakThis]
+						{
+							if (WeakThis.IsValid())
+							{
+								WeakThis->OnSettingsChanged.Broadcast(WeakThis.Get());
+								WeakThis->DeferredSettingsChangedTimerHandle.Invalidate();
+							}
+						}));
+				}
+			}
+			else
+			{
+				OnSettingsChanged.Broadcast(this);
+			}
 		}
 		
 		UE_LOG(LogEnhancedInput, Verbose, TEXT("[UEnhancedInputUserSettings::MapPlayerKey] Reset keymapping to default: '%s'"), *FoundMapping->ToString());
