Fix a bug where not all keys would be consumed when an input action was flagged to consume input, only the last key mapping would be

#rnx

diff --git a/Engine/Plugins/EnhancedInput/Source/EnhancedInput/Private/EnhancedInputSubsystemInterface.cpp b/Engine/Plugins/EnhancedInput/Source/EnhancedInput/Private/EnhancedInputSubsystemInterface.cpp
--- a/Engine/Plugins/EnhancedInput/Source/EnhancedInput/Private/EnhancedInputSubsystemInterface.cpp
+++ b/Engine/Plugins/EnhancedInput/Source/EnhancedInput/Private/EnhancedInputSubsystemInterface.cpp
@@ -928,7 +928,9 @@ void IEnhancedInputSubsystemInterface::RebuildControlMappings()
 				// This has the affect where the base UPlayerInput class will not fire any legacy bindings
 				if (Mapping.Action->bConsumesActionAndAxisMappings)
 				{
-					PlayerInput->KeyConsumptionData.Add(Mapping.Action, { Mapping.Key, static_cast<ETriggerEvent>(Mapping.Action->TriggerEventsThatConsumeLegacyKeys) });	
+					FKeyConsumptionOptions& Opts = PlayerInput->KeyConsumptionData.FindOrAdd(Mapping.Action);
+					Opts.KeysToConsume.AddUnique(Mapping.Key);
+					Opts.EventsToCauseConsumption |= static_cast<ETriggerEvent>(Mapping.Action->TriggerEventsThatConsumeLegacyKeys);
 				}
 				
 				if (Mapping.Action && !AppliedKeys.Contains(Mapping.Key))
diff --git a/Engine/Plugins/EnhancedInput/Source/EnhancedInput/Private/EnhancedPlayerInput.cpp b/Engine/Plugins/EnhancedInput/Source/EnhancedInput/Private/EnhancedPlayerInput.cpp
--- a/Engine/Plugins/EnhancedInput/Source/EnhancedInput/Private/EnhancedPlayerInput.cpp
+++ b/Engine/Plugins/EnhancedInput/Source/EnhancedInput/Private/EnhancedPlayerInput.cpp
@@ -606,7 +606,11 @@ void UEnhancedPlayerInput::EvaluateInputDelegates(const TArray<UInputComponent*>
 					{
 						if (static_cast<uint8>(ConsumptionData->EventsToCauseConsumption & Delegate->GetTriggerEvent()) != 0)
 						{
-							ConsumeKey(ConsumptionData->KeyToConsume);
+							// Consume all keys that are mapped to this input action with the proper trigger values
+							for (const FKey& KeyToConsume : ConsumptionData->KeysToConsume)
+							{
+								ConsumeKey(KeyToConsume);	
+							}
 						}
 					}
 					Delegate->Execute(*ActionData);
diff --git a/Engine/Plugins/EnhancedInput/Source/EnhancedInput/Public/EnhancedPlayerInput.h b/Engine/Plugins/EnhancedInput/Source/EnhancedInput/Public/EnhancedPlayerInput.h
--- a/Engine/Plugins/EnhancedInput/Source/EnhancedInput/Public/EnhancedPlayerInput.h
+++ b/Engine/Plugins/EnhancedInput/Source/EnhancedInput/Public/EnhancedPlayerInput.h
@@ -31,8 +31,8 @@ struct FKeyConsumptionOptions
 {
 	GENERATED_BODY()
 	
-	/** The key that should be consumed if the trigger state is reached */
-	FKey KeyToConsume = EKeys::Invalid;
+	/** Keys that should be consumed if the trigger state is reached */
+	TArray<FKey> KeysToConsume;
 		
 	/** A bitmask of trigger events that when reached, should cause the key to be marked as consumed. */
 	ETriggerEvent EventsToCauseConsumption = ETriggerEvent::None;
diff --git a/Engine/Source/Runtime/Engine/Private/UserInterface/PlayerInput.cpp b/Engine/Source/Runtime/Engine/Private/UserInterface/PlayerInput.cpp
--- a/Engine/Source/Runtime/Engine/Private/UserInterface/PlayerInput.cpp
+++ b/Engine/Source/Runtime/Engine/Private/UserInterface/PlayerInput.cpp
@@ -2052,11 +2052,8 @@ void UPlayerInput::Tick(float DeltaTime)
 
 void UPlayerInput::ConsumeKey(FKey Key)
 {
-	FKeyState* const KeyState = KeyStateMap.Find(Key);
-	if (KeyState)
-	{
-		KeyState->bConsumed = true;
-	}
+	FKeyState& KeyState = KeyStateMap.FindOrAdd(Key);
+	KeyState.bConsumed = true;
 }
 
 bool UPlayerInput::KeyEventOccurred(FKey Key, EInputEvent Event, TArray<uint32>& InEventIndices, const FKeyState* KeyState) const
