Make the Enhanced Input user settings interface use const pointers to UInputMappingContext objects to be consistent with the rest of the EI subsystem.

Ad a flag to register the IMC with the settings when you add or remove the mapping context.

diff --git a/Engine/Plugins/EnhancedInput/Source/EnhancedInput/Private/EnhancedInputSubsystemInterface.cpp b/Engine/Plugins/EnhancedInput/Source/EnhancedInput/Private/EnhancedInputSubsystemInterface.cpp
--- a/Engine/Plugins/EnhancedInput/Source/EnhancedInput/Private/EnhancedInputSubsystemInterface.cpp
+++ b/Engine/Plugins/EnhancedInput/Source/EnhancedInput/Private/EnhancedInputSubsystemInterface.cpp
@@ -184,6 +184,14 @@ void IEnhancedInputSubsystemInterface::AddMappingContext(const UInputMappingCont
 			PlayerInput->AppliedInputContexts.Add(MappingContext, Priority);
 			RequestRebuildControlMappings(Options);
 		}
+
+		if (Options.bNotifyUserSettings)
+		{
+			if (UEnhancedInputUserSettings* Settings = GetUserSettings())
+			{
+				Settings->RegisterInputMappingContext(MappingContext);
+			}
+		}
 	}
 	else
 	{
@@ -200,6 +208,14 @@ void IEnhancedInputSubsystemInterface::RemoveMappingContext(const UInputMappingC
 			PlayerInput->AppliedInputContexts.Remove(MappingContext);
 			RequestRebuildControlMappings(Options);
 		}
+
+		if (Options.bNotifyUserSettings)
+		{
+			if (UEnhancedInputUserSettings* Settings = GetUserSettings())
+			{
+				Settings->UnregisterInputMappingContext(MappingContext);
+			}
+		}
 	}
 }
 
diff --git a/Engine/Plugins/EnhancedInput/Source/EnhancedInput/Private/UserSettings/EnhancedInputUserSettings.cpp b/Engine/Plugins/EnhancedInput/Source/EnhancedInput/Private/UserSettings/EnhancedInputUserSettings.cpp
--- a/Engine/Plugins/EnhancedInput/Source/EnhancedInput/Private/UserSettings/EnhancedInputUserSettings.cpp
+++ b/Engine/Plugins/EnhancedInput/Source/EnhancedInput/Private/UserSettings/EnhancedInputUserSettings.cpp
@@ -805,7 +805,7 @@ void UEnhancedInputUserSettings::Serialize(FArchive& Ar)
 				}
 
 				// We need to populate this key profile with all the known key mappings so that it's up to date 
-				for (TObjectPtr<UInputMappingContext> IMC : RegisteredMappingContexts)
+				for (TObjectPtr<const UInputMappingContext> IMC : RegisteredMappingContexts)
 				{
 					RegisterKeyMappingsToProfile(*NewProfile, IMC);
 				}
@@ -1218,7 +1218,7 @@ UEnhancedPlayerMappableKeyProfile* UEnhancedInputUserSettings::CreateNewKeyProfi
 		SavedKeyProfiles.Add(InArgs.ProfileIdentifier, OutProfile);
 		
 		// We need to populate this key profile with all the known key mappings so that it's up to date 
-		for (TObjectPtr<UInputMappingContext> IMC : RegisteredMappingContexts)
+		for (TObjectPtr<const UInputMappingContext> IMC : RegisteredMappingContexts)
 		{
 			RegisterKeyMappingsToProfile(*OutProfile, IMC);
 		}
@@ -1256,7 +1256,7 @@ bool UEnhancedInputUserSettings::RegisterInputMappingContexts(const TSet<UInputM
 	return bResult;
 }
 
-bool UEnhancedInputUserSettings::RegisterInputMappingContext(UInputMappingContext* IMC)
+bool UEnhancedInputUserSettings::RegisterInputMappingContext(const UInputMappingContext* IMC)
 {
 	if (!IMC)
 	{
@@ -1290,7 +1290,7 @@ bool UEnhancedInputUserSettings::RegisterInputMappingContext(UInputMappingContex
 	return bResult;
 }
 
-bool UEnhancedInputUserSettings::RegisterKeyMappingsToProfile(UEnhancedPlayerMappableKeyProfile& Profile, UInputMappingContext* IMC)
+bool UEnhancedInputUserSettings::RegisterKeyMappingsToProfile(UEnhancedPlayerMappableKeyProfile& Profile, const UInputMappingContext* IMC)
 {
 	if (!IMC)
 	{
@@ -1413,7 +1413,7 @@ bool UEnhancedInputUserSettings::UnregisterInputMappingContexts(const TSet<UInpu
 	return bResult;
 }
 
-const TSet<TObjectPtr<UInputMappingContext>>& UEnhancedInputUserSettings::GetRegisteredInputMappingContexts() const
+const TSet<TObjectPtr<const UInputMappingContext>>& UEnhancedInputUserSettings::GetRegisteredInputMappingContexts() const
 {
 	return RegisteredMappingContexts;
 }
diff --git a/Engine/Plugins/EnhancedInput/Source/EnhancedInput/Public/EnhancedInputSubsystemInterface.h b/Engine/Plugins/EnhancedInput/Source/EnhancedInput/Public/EnhancedInputSubsystemInterface.h
--- a/Engine/Plugins/EnhancedInput/Source/EnhancedInput/Public/EnhancedInputSubsystemInterface.h
+++ b/Engine/Plugins/EnhancedInput/Source/EnhancedInput/Public/EnhancedInputSubsystemInterface.h
@@ -51,6 +51,7 @@ struct FModifyContextOptions
 	FModifyContextOptions()
 		: bIgnoreAllPressedKeysUntilRelease(true)
 		, bForceImmediately(false)
+		, bNotifyUserSettings(false)
 	{}
 
 	// If true then any keys that are pressed during the rebuild of control mappings will be ignored until they are released.
@@ -60,6 +61,11 @@ struct FModifyContextOptions
 	// The mapping changes will be applied synchronously, rather than at the end of the frame, making them available to the input system on the same frame.
 	UPROPERTY(EditAnywhere, BlueprintReadWrite, Category = "Input")
 	uint8 bForceImmediately : 1;
+	
+	// If true, then this Mapping Context will be registered or unregistered with the Enhanced Input User Settings on this subsystem,
+	// if they exist.
+	UPROPERTY(EditAnywhere, BlueprintReadWrite, Category = "Input")
+	uint8 bNotifyUserSettings : 1;
 };
 
 // Includes native functionality shared between all subsystems
diff --git a/Engine/Plugins/EnhancedInput/Source/EnhancedInput/Public/UserSettings/EnhancedInputUserSettings.h b/Engine/Plugins/EnhancedInput/Source/EnhancedInput/Public/UserSettings/EnhancedInputUserSettings.h
--- a/Engine/Plugins/EnhancedInput/Source/EnhancedInput/Public/UserSettings/EnhancedInputUserSettings.h
+++ b/Engine/Plugins/EnhancedInput/Source/EnhancedInput/Public/UserSettings/EnhancedInputUserSettings.h
@@ -660,7 +660,7 @@ public:
 	 * in the context and create an initial Player Mappable Key for every mapping that is marked as mappable.
 	 */
 	UFUNCTION(BlueprintCallable, Category="Enhanced Input|User Settings")
-	virtual bool RegisterInputMappingContext(UInputMappingContext* IMC);
+	virtual bool RegisterInputMappingContext(const UInputMappingContext* IMC);
 
 	/** Registers multiple mapping contexts with the settings */
 	UFUNCTION(BlueprintCallable, Category="Enhanced Input|User Settings")
@@ -675,7 +675,7 @@ public:
 	bool UnregisterInputMappingContexts(const TSet<UInputMappingContext*>& MappingContexts);
 
 	/** Gets all the currently registered mapping contexts with the settings */
-	const TSet<TObjectPtr<UInputMappingContext>>& GetRegisteredInputMappingContexts() const;
+	const TSet<TObjectPtr<const UInputMappingContext>>& GetRegisteredInputMappingContexts() const;
 
 	/** Returns true if this mapping context is currently registered with the settings */
 	UFUNCTION(BlueprintCallable, BlueprintPure, Category="Enhanced Input|User Settings")
@@ -695,7 +695,7 @@ protected:
 	 * This will be called for each key profile when users initially call RegisterInputMappingContext, and on
 	 * load during Serialize to correctly populate a loaded profile with all the known registered IMC's mappings.
 	 */
-	virtual bool RegisterKeyMappingsToProfile(UEnhancedPlayerMappableKeyProfile& Profile, UInputMappingContext* IMC);
+	virtual bool RegisterKeyMappingsToProfile(UEnhancedPlayerMappableKeyProfile& Profile, const UInputMappingContext* IMC);
 
 	/**
 	 * Determines the hardware device that an action key mapping is associated with. By default, the hardware will be FHardwareDeviceIdentifier::Invalid.
@@ -731,5 +731,5 @@ protected:
 	 * active on the user, but you want to track for creating a menu for key mappings.
 	 */
 	UPROPERTY(Transient)
-	TSet<TObjectPtr<UInputMappingContext>> RegisteredMappingContexts;
+	TSet<TObjectPtr<const UInputMappingContext>> RegisteredMappingContexts;
 };
