Chord action input trigger and legacy input consumption tweaks:
- Fix chord action input triggers to block other input triggers unless triggered
- Enable an input action binding to block legacy input for any specified trigger states, opt-in only

diff --git a/Engine/Plugins/EnhancedInput/Source/EnhancedInput/Private/EnhancedPlayerInput.cpp b/Engine/Plugins/EnhancedInput/Source/EnhancedInput/Private/EnhancedPlayerInput.cpp
--- a/Engine/Plugins/EnhancedInput/Source/EnhancedInput/Private/EnhancedPlayerInput.cpp
+++ b/Engine/Plugins/EnhancedInput/Source/EnhancedInput/Private/EnhancedPlayerInput.cpp
@@ -570,6 +570,19 @@ void UEnhancedPlayerInput::EvaluateInputDelegates(const TArray<UInputComponent*>
 						TriggeredActionsThisTick.Add(ActionData->GetSourceAction());
 					}
 				}
+
+				// If this delegate is bound to an action that has an event and is flagged to consume legacy keys, then mark it as such.
+				if (const FKeyConsumptionOptions* const ConsumptionData = KeyConsumptionData.Find(ActionData->GetSourceAction()))
+				{
+					if (static_cast<uint8>(ConsumptionData->EventsToCauseConsumption & ActionData->TriggerEvent) != 0)
+					{
+						// Consume all keys that are mapped to this input action with the proper trigger values
+						for (const FKey& KeyToConsume : ConsumptionData->KeysToConsume)
+						{
+							ConsumeKey(KeyToConsume);
+						}
+					}
+				}
 			}
 		}
 
@@ -601,18 +614,6 @@ void UEnhancedPlayerInput::EvaluateInputDelegates(const TArray<UInputComponent*>
 				// Search for the action instance data a second time as a previous delegate call may have deleted it.
 				if (const FInputActionInstance* ActionData = FindActionInstanceData(DelegateAction))
 				{
-					// If this enhanced input delegate has triggered and is flagged to consume legacy keys, then mark it as such.
-					if (const FKeyConsumptionOptions* ConsumptionData = KeyConsumptionData.Find(ActionData->GetSourceAction()))
-					{
-						if (static_cast<uint8>(ConsumptionData->EventsToCauseConsumption & Delegate->GetTriggerEvent()) != 0)
-						{
-							// Consume all keys that are mapped to this input action with the proper trigger values
-							for (const FKey& KeyToConsume : ConsumptionData->KeysToConsume)
-							{
-								ConsumeKey(KeyToConsume);	
-							}
-						}
-					}
 					Delegate->Execute(*ActionData);
 				}
 			}
diff --git a/Engine/Plugins/EnhancedInput/Source/EnhancedInput/Private/InputAction.cpp b/Engine/Plugins/EnhancedInput/Source/EnhancedInput/Private/InputAction.cpp
--- a/Engine/Plugins/EnhancedInput/Source/EnhancedInput/Private/InputAction.cpp
+++ b/Engine/Plugins/EnhancedInput/Source/EnhancedInput/Private/InputAction.cpp
@@ -72,10 +72,11 @@ ETriggerState FTriggerStateTracker::EvaluateTriggers(const UEnhancedPlayerInput*
 			bFoundActiveTrigger |= (CurrentState != ETriggerState::None);
 			break;
 		case ETriggerType::Blocker:
-			bBlocking |= (CurrentState == ETriggerState::Triggered);
-			// Ongoing blockers don't count as active triggers
+			// Ignore dedicated blockers, IsBlocking handles the state update
 			break;
 		}
+
+		bBlocking |= Trigger->IsBlocking(CurrentState);
 	}
 
 	return GetState();
diff --git a/Engine/Plugins/EnhancedInput/Source/EnhancedInput/Public/InputTriggers.h b/Engine/Plugins/EnhancedInput/Source/EnhancedInput/Public/InputTriggers.h
--- a/Engine/Plugins/EnhancedInput/Source/EnhancedInput/Public/InputTriggers.h
+++ b/Engine/Plugins/EnhancedInput/Source/EnhancedInput/Public/InputTriggers.h
@@ -66,8 +66,8 @@ enum class ETriggerType : uint8
 	// Input may trigger only if all implicit triggers are triggered.
 	Implicit,
 
-	// Inverted trigger that will block all other triggers if it is triggered.
-	Blocker,
+	// Inverted trigger that will block all other triggers when IsBlocking returns true
+	Blocker
 };
 
 /**
@@ -162,6 +162,9 @@ public:
 	 * Returns true if the the ETriggerEvent is can be triggered based off of the ETriggerEventsSupported.
 	 */
 	static bool IsSupportedTriggerEvent(const ETriggerEventsSupported SupportedEvents, const ETriggerEvent Event);
+
+	/** Returns true when all other triggers should be blocked, i.e. chords block when not triggered. */
+	virtual bool IsBlocking(const ETriggerState State) const { return false; }
 		
 	// Provide debug output for use with ShowDebug EnhancedInput. Return an empty string to disable display.
 	virtual FString GetDebugState() const { return FString(); }
@@ -395,6 +398,8 @@ protected:
 	virtual ETriggerType GetTriggerType_Implementation() const override { return ETriggerType::Implicit; }
 
 	virtual ETriggerState UpdateState_Implementation(const UEnhancedPlayerInput* PlayerInput, FInputActionValue ModifiedValue, float DeltaTime) override;
+	virtual bool IsBlocking(const ETriggerState State) const override { return State != ETriggerState::Triggered; }
+
 public:
 
 	// The action that must be triggering for this trigger's action to trigger
@@ -412,6 +417,7 @@ class UInputTriggerChordBlocker final : public UInputTriggerChordAction
 	GENERATED_BODY()
 protected:
 	virtual ETriggerType GetTriggerType_Implementation() const override { return ETriggerType::Blocker; }
+	virtual bool IsBlocking(const ETriggerState State) const override { return State == ETriggerState::Triggered; }
 };
 
 USTRUCT(BlueprintType)
