Deprecate UPlayerMappableInputConfig in favor of the new UEnhancedInputUserSettings with 5.3.

diff --git a/Engine/Plugins/EnhancedInput/Source/EnhancedInput/Private/PlayerMappableInputConfig.cpp b/Engine/Plugins/EnhancedInput/Source/EnhancedInput/Private/PlayerMappableInputConfig.cpp
--- a/Engine/Plugins/EnhancedInput/Source/EnhancedInput/Private/PlayerMappableInputConfig.cpp
+++ b/Engine/Plugins/EnhancedInput/Source/EnhancedInput/Private/PlayerMappableInputConfig.cpp
@@ -11,6 +11,8 @@
 
 #define LOCTEXT_NAMESPACE "PlayerMappableInputConfig"
 
+PRAGMA_DISABLE_DEPRECATION_WARNINGS
+
 UPlayerMappableInputConfig::UPlayerMappableInputConfig(const FObjectInitializer& ObjectInitializer)
 	: Super(ObjectInitializer)
 	, ConfigName(GetFName())
@@ -130,5 +132,6 @@ TArray<FEnhancedActionKeyMapping> UPlayerMappableInputConfig::GetKeysBoundToActi
 	return OutMappings;
 }
 
-#undef LOCTEXT_NAMESPACE
+PRAGMA_ENABLE_DEPRECATION_WARNINGS
 
+#undef LOCTEXT_NAMESPACE
\ No newline at end of file
diff --git a/Engine/Plugins/EnhancedInput/Source/EnhancedInput/Public/PlayerMappableInputConfig.h b/Engine/Plugins/EnhancedInput/Source/EnhancedInput/Public/PlayerMappableInputConfig.h
--- a/Engine/Plugins/EnhancedInput/Source/EnhancedInput/Public/PlayerMappableInputConfig.h
+++ b/Engine/Plugins/EnhancedInput/Source/EnhancedInput/Public/PlayerMappableInputConfig.h
@@ -20,7 +20,7 @@ class UInputMappingContext;
  * Populate this data asset with Input Mapping Contexts that have player mappable actions in them. 
  */
 UCLASS(BlueprintType, Meta = (DisplayName = "Player Mappable Input Config", ShortTooltip = "Data asset used to define a set of player mappable controller/keyboard mappings."))
-class ENHANCEDINPUT_API UPlayerMappableInputConfig : public UPrimaryDataAsset
+class UE_DEPRECATED(5.3, "UPlayerMappableInputConfig has been deprecated. Please use UEnhancedInputUserSettings instead.") ENHANCEDINPUT_API UPlayerMappableInputConfig : public UPrimaryDataAsset
 {
 	GENERATED_BODY()
 
