Update the comments on the Enhanced Action Key Mapping and the Input Action asset to be very explict about the order of which modifiers are applied.

diff --git a/Engine/Plugins/EnhancedInput/Source/EnhancedInput/Public/EnhancedActionKeyMapping.h b/Engine/Plugins/EnhancedInput/Source/EnhancedInput/Public/EnhancedActionKeyMapping.h
--- a/Engine/Plugins/EnhancedInput/Source/EnhancedInput/Public/EnhancedActionKeyMapping.h
+++ b/Engine/Plugins/EnhancedInput/Source/EnhancedInput/Public/EnhancedActionKeyMapping.h
@@ -136,6 +136,9 @@ struct ENHANCEDINPUT_API FEnhancedActionKeyMapping
 	/**
 	* Modifiers applied to the raw key value.
 	* These are applied sequentially in array order.
+	* 
+	* Note: Modifiers defined in individual key mappings will be applied before those defined in the Input Action asset.
+	* Modifiers will not override any that are defined on the Input Action asset, they will be combined together during evaluation.
 	*/
 	UPROPERTY(EditAnywhere, Instanced, BlueprintReadWrite, Category = "Input")
 	TArray<TObjectPtr<UInputModifier>> Modifiers;
diff --git a/Engine/Plugins/EnhancedInput/Source/EnhancedInput/Public/InputAction.h b/Engine/Plugins/EnhancedInput/Source/EnhancedInput/Public/InputAction.h
--- a/Engine/Plugins/EnhancedInput/Source/EnhancedInput/Public/InputAction.h
+++ b/Engine/Plugins/EnhancedInput/Source/EnhancedInput/Public/InputAction.h
@@ -136,6 +136,8 @@ public:
 	* Modifiers are applied to the final action value.
 	* These are applied sequentially in array order.
 	* They are applied on top of any FEnhancedActionKeyMapping modifiers that drove the initial input
+	* 
+	* Note: Modifiers defined in the Input Action asset will be applied AFTER any modifiers defined in individual key mappings in the Input Mapping Context asset.
 	*/
 	UPROPERTY(EditAnywhere, Instanced, BlueprintReadWrite, Category = Action, meta=(DisplayAfter="Triggers"))
 	TArray<TObjectPtr<UInputModifier>> Modifiers;
