Make the player mappable key settings object protected instead of private. This will let people make custom UInputAction subclasses that modify it more easily.

#rnx

diff --git a/Engine/Plugins/EnhancedInput/Source/EnhancedInput/Public/InputAction.h b/Engine/Plugins/EnhancedInput/Source/EnhancedInput/Public/InputAction.h
--- a/Engine/Plugins/EnhancedInput/Source/EnhancedInput/Public/InputAction.h
+++ b/Engine/Plugins/EnhancedInput/Source/EnhancedInput/Public/InputAction.h
@@ -142,7 +142,7 @@ public:
 	UPROPERTY(EditAnywhere, Instanced, BlueprintReadWrite, Category = Action, meta=(DisplayAfter="Triggers"))
 	TArray<TObjectPtr<UInputModifier>> Modifiers;
 
-private:
+protected:
 
 	/**
 	* Holds setting information about this Action Input for setting screen and save purposes.
