Make it just a log instead of a warning when you attempt to set the key profile with a null local player

#rnx

diff --git a/Engine/Plugins/EnhancedInput/Source/EnhancedInput/Private/UserSettings/EnhancedInputUserSettings.cpp b/Engine/Plugins/EnhancedInput/Source/EnhancedInput/Private/UserSettings/EnhancedInputUserSettings.cpp
--- a/Engine/Plugins/EnhancedInput/Source/EnhancedInput/Private/UserSettings/EnhancedInputUserSettings.cpp
+++ b/Engine/Plugins/EnhancedInput/Source/EnhancedInput/Private/UserSettings/EnhancedInputUserSettings.cpp
@@ -1124,7 +1124,7 @@ bool UEnhancedInputUserSettings::SetKeyProfile(const FGameplayTag& InProfileId)
 {
 	if (!GetLocalPlayer())
 	{
-		UE_LOG(LogEnhancedInput, Error, TEXT("Failed to find the Local Player associated with the Enhanced Input user settings!"), *InProfileId.ToString());
+		UE_LOG(LogEnhancedInput, Log, TEXT("Failed to find the Local Player associated with the Enhanced Input user settings!"), *InProfileId.ToString());
 		return false;
 	}
 
