Add Unit tests for player mappable keys

#rnx

diff --git a/Engine/Plugins/EnhancedInput/Source/InputEditor/InputEditor.Build.cs b/Engine/Plugins/EnhancedInput/Source/InputEditor/InputEditor.Build.cs
--- a/Engine/Plugins/EnhancedInput/Source/InputEditor/InputEditor.Build.cs
+++ b/Engine/Plugins/EnhancedInput/Source/InputEditor/InputEditor.Build.cs
@@ -33,6 +33,7 @@ namespace UnrealBuildTool.Rules
 					"ToolMenus",
 					"ContentBrowser",
 					"SourceControl",
+					"GameplayTags",
                 }
             );
         }
diff --git a/Engine/Plugins/EnhancedInput/Source/InputEditor/Private/Tests/InputPlayerMappableKeysTests.cpp b/Engine/Plugins/EnhancedInput/Source/InputEditor/Private/Tests/InputPlayerMappableKeysTests.cpp
new file mode 100644
--- /dev/null
+++ b/Engine/Plugins/EnhancedInput/Source/InputEditor/Private/Tests/InputPlayerMappableKeysTests.cpp
@@ -0,0 +1,408 @@
+// Copyright Epic Games, Inc. All Rights Reserved.
+
+#include <ThirdParty/Perforce/p4api-2018.1/Include/Win64/VS2015/p4/error.h>
+
+#include "InputMappingContext.h"
+#include "InputTestFramework.h"
+#include "Misc/AutomationTest.h"
+#include "UserSettings/EnhancedInputUserSettings.h"
+#include "Algo/Count.h"
+
+// Tests focused Player Remappable Keys
+
+namespace UE::Input
+{
+	static FKey TestKeyA = EKeys::A;
+	static FKey TestKeyB = EKeys::B;
+	static FKey TestKeyC = EKeys::C;
+
+	static FKey TestKeyQ = EKeys::Q;
+	static FKey TestKeyX = EKeys::X;
+	static FKey TestKeyY = EKeys::Y;
+	static FKey TestKeyZ = EKeys::Z;
+	
+	static FName TestContext_2 = TEXT("TestContext_2");
+}
+
+#define REMAPPABLEKEY_SUBTEST(DESC) \
+	for(FString ScopedSubTestDescription = TEXT(DESC);ScopedSubTestDescription.Len();ScopedSubTestDescription = "")	// Bodge to create a scoped test description. Usage: REMAPPABLEKEY_SUBTEST("My Test Description") { TestCode... TriggerStateIs(ETriggerState::Triggered); }
+
+constexpr auto BasicPlayerMappableKeysTestFlags = EAutomationTestFlags::EditorContext | EAutomationTestFlags::EngineFilter;
+
+UControllablePlayer& ABasicPlayerMappableKeysTest(FAutomationTestBase* Test)
+{
+	UWorld* World =
+	GIVEN(AnEmptyWorld());
+
+	UControllablePlayer& Data =
+	AND(AControllablePlayer(World));
+	Test->TestTrue(TEXT("Controllable Player is valid"), Data.IsValid());
+
+	UEnhancedInputUserSettings* Settings = Data.Subsystem->GetUserSettings();
+	Test->TestNotNull("Mocked User Settings", Settings);
+
+	return Data;
+}
+
+// Provides a test player with a single input action (TestAction) added to it mapped to X
+UControllablePlayer& APlayerMappableKeysTest_WithAction(FAutomationTestBase* Test)
+{
+	UWorld* World =
+	GIVEN(AnEmptyWorld());
+
+	UControllablePlayer& Data =
+	AND(AControllablePlayer(World));
+	Test->TestTrue(TEXT("Controllable Player is valid"), Data.IsValid());
+
+	UEnhancedInputUserSettings* Settings = Data.Subsystem->GetUserSettings();
+	Test->TestNotNull("Mocked User Settings", Settings);
+
+	// Test context
+	AND(AnInputContextIsAppliedToAPlayer(Data, TestContext, 0));
+	
+	TObjectPtr<UInputMappingContext>* IMC = Data.InputContext.Find(TestContext);
+	Test->TestNotNull("Mock Mapping Context", IMC);
+
+	// Add a simple mapping from test action to X
+	UInputAction* Action1 =
+	AND(AnInputAction(Data, TestAction, EInputActionValueType::Axis2D));
+
+	Test->TestNotNull("Mock Input Action", Action1);
+
+	// Test mapping to X
+	AND(AnActionIsMappedToAKey(Data, TestContext, TestAction, UE::Input::TestKeyX));
+	
+	const bool bRes = Settings->RegisterInputMappingContext(*IMC);
+	Test->TestTrue("Registered Mock IMC", bRes);
+
+	return Data;
+}
+
+struct FMockKeyMappingData
+{
+	FName ActionName;
+	FName ContextName;
+	FKey DefaultKey;
+	EInputActionValueType ValueType;
+	EPlayerMappableKeySlot ExpectedSlot;
+	FName HardwareDeviceId;
+};
+
+
+IMPLEMENT_SIMPLE_AUTOMATION_TEST(FRegisterIMCTest, "Input.PlayerMappableKeys.Registration", BasicPlayerMappableKeysTestFlags)
+bool FRegisterIMCTest::RunTest(const FString& Parameters)
+{
+	UControllablePlayer& Data =
+	GIVEN(ABasicPlayerMappableKeysTest(this));
+
+	UEnhancedInputUserSettings* Settings = Data.Subsystem->GetUserSettings();
+	TestNotNull("Mock Settings", Settings);
+
+	static const TArray<FMockKeyMappingData> IMC_Mappings =
+	{
+		// TestContext
+		{ TestAction,  TestContext,  UE::Input::TestKeyX, EInputActionValueType::Axis1D, EPlayerMappableKeySlot::First },
+		{ TestAction2, TestContext,  UE::Input::TestKeyY, EInputActionValueType::Axis2D, EPlayerMappableKeySlot::First },
+		{ TestAction3, TestContext,  UE::Input::TestKeyZ, EInputActionValueType::Axis3D, EPlayerMappableKeySlot::First },
+		{ TestAction3, TestContext,  UE::Input::TestKeyQ, EInputActionValueType::Axis3D, EPlayerMappableKeySlot::Second },
+
+		// UE::Input::TestContext_2
+		{ TestAction4, UE::Input::TestContext_2,  UE::Input::TestKeyA, EInputActionValueType::Axis1D, EPlayerMappableKeySlot::First },
+		{ TestAction5, UE::Input::TestContext_2,  UE::Input::TestKeyB, EInputActionValueType::Axis2D, EPlayerMappableKeySlot::First },
+		{ TestAction6, UE::Input::TestContext_2,  UE::Input::TestKeyC, EInputActionValueType::Axis3D, EPlayerMappableKeySlot::First },
+	};
+	
+	AND(AnInputContextIsAppliedToAPlayer(Data, TestContext, 0));
+	AND(AnInputContextIsAppliedToAPlayer(Data, UE::Input::TestContext_2, 0));
+
+	for (const FMockKeyMappingData& KeyData : IMC_Mappings)
+	{
+		// TestContext
+		UInputAction* Action1 =
+		AND(AnInputAction(Data, KeyData.ActionName, KeyData.ValueType));
+		AND(AnActionIsMappedToAKey(Data, KeyData.ContextName, KeyData.ActionName, KeyData.DefaultKey));
+	}
+	
+	// Ensure the state of the key is correct by default.
+	TObjectPtr<UInputMappingContext>* IMC_1 = Data.InputContext.Find(TestContext);
+	TestNotNull("Mock Mapping Context 1", IMC_1);
+
+	TObjectPtr<UInputMappingContext>* IMC_2 = Data.InputContext.Find(UE::Input::TestContext_2);
+	TestNotNull("Mock Mapping Context 2", IMC_2);
+
+	// You can register one IMC
+	REMAPPABLEKEY_SUBTEST("Register Mapping Context")
+	{
+		const bool bRes = Settings->RegisterInputMappingContext(*IMC_1);
+		TestTrue("Register Mock Input Mapping Context", bRes);
+
+		// We only expect one context to be registered at this time
+		const TSet<TObjectPtr<const UInputMappingContext>>& RegisteredContexts = Settings->GetRegisteredInputMappingContexts();
+		TestEqual("Num Registered Contexts", RegisteredContexts.Num(), 1);
+
+		// We should have 3 registered key mappings here, one for each action that was mapped in the IMC
+		const UEnhancedPlayerMappableKeyProfile* Profile = Settings->GetCurrentKeyProfile();
+		TestNotNull("Current Key Profile", Profile);
+
+		// We expect 3 mapping rows, one for each mapping in the IMC
+		const TMap<FName, FKeyMappingRow>& Rows = Profile->GetPlayerMappingRows();
+		TestEqual("Key Mapping Rows", Rows.Num(), 3);
+	}
+
+	// You cannot register that IMC over again
+	REMAPPABLEKEY_SUBTEST("Cannot re-register a context")
+	{
+		const bool bRes = Settings->RegisterInputMappingContext(*IMC_1);
+		TestFalse("Re-registering the same IMC should be false", bRes);
+
+		// We only expect one context to be registered at this time
+		const TSet<TObjectPtr<const UInputMappingContext>>& RegisteredContexts = Settings->GetRegisteredInputMappingContexts();
+		TestEqual("Num Registered Contexts", RegisteredContexts.Num(), 1);
+	}
+	
+	REMAPPABLEKEY_SUBTEST("Register Multiple Mapping Contexts")
+	{
+		const bool bRes = Settings->RegisterInputMappingContext(*IMC_2);
+		TestTrue("Register Mock Input Mapping Context 2", bRes);
+
+		const TSet<TObjectPtr<const UInputMappingContext>>& RegisteredContexts = Settings->GetRegisteredInputMappingContexts();
+		TestEqual("Num Registered Contexts", RegisteredContexts.Num(), 2);
+
+		UEnhancedPlayerMappableKeyProfile* Profile = Settings->GetCurrentKeyProfile();
+		TestNotNull("Current Key Profile", Profile);
+		
+		const TMap<FName, FKeyMappingRow>& Rows = Profile->GetPlayerMappingRows();
+
+		// We should have 6 rows, one for each test action
+		TestEqual("Key Mapping Rows", Rows.Num(), 6);
+	}
+	
+	REMAPPABLEKEY_SUBTEST("Key Mappings Have Correct Number or registered mappings")
+	{
+		UEnhancedPlayerMappableKeyProfile* Profile = Settings->GetCurrentKeyProfile();
+		TestNotNull("Current Key Profile", Profile);
+		
+		for (const FMockKeyMappingData& KeyData : IMC_Mappings)
+		{
+			const FKeyMappingRow* Row = Profile->FindKeyMappingRow(KeyData.ActionName);
+			TestNotNull("Mapping Row", Row);
+
+			UInputMappingContext* IMC = FInputTestHelper::FindContext(Data, KeyData.ContextName);
+			UInputAction* Action = FInputTestHelper::FindAction(Data, KeyData.ActionName);
+
+			// The row should have the same number of player mappings as there are mappings for each input action
+			auto IsCorrectAction = [&Action](const FEnhancedActionKeyMapping& Mapping) { return Mapping.Action == Action; };
+			const int32 NumMappingsToAction = Algo::CountIf(IMC->GetMappings(), IsCorrectAction);
+			TestEqual("Correct Number of Rows", Row->Mappings.Num(), NumMappingsToAction);
+			
+			// Test that the query results are as expected
+			FPlayerMappableKeyQueryOptions Opts = {};
+			Opts.MappingName = KeyData.ActionName;
+			Opts.KeyToMatch = KeyData.DefaultKey;
+			Opts.SlotToMatch = KeyData.ExpectedSlot;
+			Opts.bMatchBasicKeyTypes = true;
+			Opts.bMatchKeyAxisType = true;
+
+			TArray<FKey> MappedKeys;
+			Profile->QueryPlayerMappedKeys(Opts, MappedKeys);
+
+			TestEqual("Number of mapped keys", MappedKeys.Num(), 1);
+			if (!MappedKeys.IsEmpty())
+			{
+				TestEqual("Mapped to the correct key", MappedKeys[0], KeyData.DefaultKey);
+			}
+		}
+	}
+	
+	// Unregister the IMC's
+	REMAPPABLEKEY_SUBTEST("Unregister IMC_1")
+	{
+		const bool bRes = Settings->UnregisterInputMappingContext(*IMC_1);
+		TestTrue("Register Mock Input Mapping Context 1", bRes);
+
+		const TSet<TObjectPtr<const UInputMappingContext>>& RegisteredContexts = Settings->GetRegisteredInputMappingContexts();
+		TestEqual("Num Registered Contexts", RegisteredContexts.Num(), 1);
+
+		UEnhancedPlayerMappableKeyProfile* Profile = Settings->GetCurrentKeyProfile();
+		TestNotNull("Current Key Profile", Profile);
+
+		// Unregistering a mapping context shouldn't actually remove it's mapping rows because they will be saved
+		const TMap<FName, FKeyMappingRow>& Rows = Profile->GetPlayerMappingRows();
+		TestEqual("Key Mapping Rows", Rows.Num(), 6);
+	}
+
+	REMAPPABLEKEY_SUBTEST("Unregister IMC_2")
+	{
+		const bool bRes = Settings->UnregisterInputMappingContext(*IMC_2);
+		TestTrue("Register Mock Input Mapping Context 2", bRes);
+
+		const TSet<TObjectPtr<const UInputMappingContext>>& RegisteredContexts = Settings->GetRegisteredInputMappingContexts();
+		TestEqual("Num Registered Contexts", RegisteredContexts.Num(), 0);
+
+		UEnhancedPlayerMappableKeyProfile* Profile = Settings->GetCurrentKeyProfile();
+		TestNotNull("Current Key Profile", Profile);
+		
+		const TMap<FName, FKeyMappingRow>& Rows = Profile->GetPlayerMappingRows();
+		TestEqual("Key Mapping Rows", Rows.Num(), 6);
+	}
+	
+	return true;
+}
+
+//////////////////////////////////////////////////////////////
+// Mapping Keys
+
+IMPLEMENT_SIMPLE_AUTOMATION_TEST(FMapPlayerKeyTest, "Input.PlayerMappableKeys.Map", BasicPlayerMappableKeysTestFlags)
+bool FMapPlayerKeyTest::RunTest(const FString& Parameters)
+{
+	UControllablePlayer& Data =
+	GIVEN(APlayerMappableKeysTest_WithAction(this));
+
+	UEnhancedInputUserSettings* Settings = Data.Subsystem->GetUserSettings();
+	TestNotNull("Mock Settings", Settings);
+
+	UEnhancedPlayerMappableKeyProfile* Profile = Settings->GetCurrentKeyProfile();
+	TestNotNull("Current Key Profile", Profile);
+
+	FGameplayTagContainer ErrorReason;
+	
+	// Remap the action in slot 1 from x to z
+	FMapPlayerKeyArgs Args = {};
+	Args.MappingName = TestAction;
+	Args.NewKey = UE::Input::TestKeyZ;
+	Args.Slot = EPlayerMappableKeySlot::First;
+	
+	Settings->MapPlayerKey(Args, OUT ErrorReason);
+
+	TestTrue("No Error Reason", ErrorReason.IsEmpty());
+
+	// Query this key now and make sure it was remapped correctly
+	FPlayerMappableKeyQueryOptions Opts = {};
+	Opts.MappingName = Args.MappingName;
+	Opts.KeyToMatch = Args.NewKey;
+	Opts.SlotToMatch = Args.Slot;
+
+	TArray<FKey> MappedKeys;
+	Profile->QueryPlayerMappedKeys(Opts, MappedKeys);
+
+	TestEqual("Number of mapped keys", MappedKeys.Num(), 1);
+	if (!MappedKeys.IsEmpty())
+	{
+		TestEqual("re-mapped to the correct key", MappedKeys[0], Args.NewKey);
+	}
+
+	// If we call map player key on a slot that doesn't exist then it should create one
+	Args.Slot = EPlayerMappableKeySlot::Third;
+	Args.NewKey = UE::Input::TestKeyA;
+	Args.bCreateMatchingSlotIfNeeded = true;
+	
+	Settings->MapPlayerKey(Args, OUT ErrorReason);
+	TestTrue("No Error Reason", ErrorReason.IsEmpty());
+
+	Opts.MappingName = Args.MappingName;
+	Opts.KeyToMatch = Args.NewKey;
+	Opts.SlotToMatch = Args.Slot;
+	Profile->QueryPlayerMappedKeys(Opts, MappedKeys);
+
+	TestEqual("Number of mapped keys", MappedKeys.Num(), 1);
+	if (!MappedKeys.IsEmpty())
+	{
+		TestEqual("Added a new slot correctly", MappedKeys[0], Args.NewKey);
+	}
+	
+	return true;
+}
+
+// Unmapping keys
+IMPLEMENT_SIMPLE_AUTOMATION_TEST(FUnmapKeyTest, "Input.PlayerMappableKeys.Unmap", BasicPlayerMappableKeysTestFlags)
+bool FUnmapKeyTest::RunTest(const FString& Parameters)
+{
+	UControllablePlayer& Data =
+	GIVEN(APlayerMappableKeysTest_WithAction(this));
+
+	UEnhancedInputUserSettings* Settings = Data.Subsystem->GetUserSettings();
+	TestNotNull("Mock Settings", Settings);
+
+	UEnhancedPlayerMappableKeyProfile* Profile = Settings->GetCurrentKeyProfile();
+	TestNotNull("Current Key Profile", Profile);
+
+	FGameplayTagContainer ErrorReason;
+	
+	// Add a second key mapping to Z
+	FMapPlayerKeyArgs Args = {};
+	Args.MappingName = TestAction;
+	Args.NewKey = UE::Input::TestKeyZ;
+	Args.Slot = EPlayerMappableKeySlot::Second;
+	
+	Settings->MapPlayerKey(Args, OUT ErrorReason);
+	TestTrue("Successful Remap", ErrorReason.IsEmpty());
+
+
+	// Unmap this key now
+	Settings->UnMapPlayerKey(Args, ErrorReason);
+	TestTrue("No Error Reason during unmap", ErrorReason.IsEmpty());
+
+	// Query this key now and make sure it was remapped correctly
+	FPlayerMappableKeyQueryOptions Opts = {};
+	Opts.MappingName = Args.MappingName;
+	Opts.KeyToMatch = Args.NewKey;
+	Opts.SlotToMatch = Args.Slot;
+
+	TArray<FKey> MappedKeys;
+	Profile->QueryPlayerMappedKeys(Opts, MappedKeys);
+
+	// The key was unmapped, so it should be set to Invalid because it is not on by default
+	TestEqual("Number of mapped keys", MappedKeys.Num(), 1);
+	if (!MappedKeys.IsEmpty())
+	{
+		TestEqual("Added a new slot correctly", MappedKeys[0], EKeys::Invalid);
+	}
+
+	return true;
+}
+
+// Reset to default
+IMPLEMENT_SIMPLE_AUTOMATION_TEST(FResetToDefaultTest, "Input.PlayerMappableKeys.ResetToDefault", BasicPlayerMappableKeysTestFlags)
+
+bool FResetToDefaultTest::RunTest(const FString& Parameters)
+{
+	UControllablePlayer& Data =
+	GIVEN(APlayerMappableKeysTest_WithAction(this));
+
+	UEnhancedInputUserSettings* Settings = Data.Subsystem->GetUserSettings();
+	TestNotNull("Mock Settings", Settings);
+
+	UEnhancedPlayerMappableKeyProfile* Profile = Settings->GetCurrentKeyProfile();
+	TestNotNull("Current Key Profile", Profile);
+
+	FGameplayTagContainer ErrorReason;
+	
+	// Remap the action in slot 1 from x to z
+	FMapPlayerKeyArgs Args = {};
+	Args.MappingName = TestAction;
+	Args.NewKey = UE::Input::TestKeyZ;
+	Args.Slot = EPlayerMappableKeySlot::First;
+	
+	Settings->MapPlayerKey(Args, OUT ErrorReason);
+	TestTrue("Successful Remap", ErrorReason.IsEmpty());
+
+	// Reset the mapping back to default
+	Profile->ResetMappingToDefault(TestAction);
+
+	// Confirm this key has been reset!
+	FPlayerMappableKeyQueryOptions Opts = {};
+	Opts.MappingName = Args.MappingName;
+	Opts.SlotToMatch = Args.Slot;
+
+	// Query the key mapping and make sure it is back to the default of X
+	TArray<FKey> MappedKeys;
+	Profile->QueryPlayerMappedKeys(Opts, MappedKeys);
+
+	TestEqual("Number of mapped keys", MappedKeys.Num(), 1);
+	if (!MappedKeys.IsEmpty())
+	{
+		TestEqual("reset to default correctly!", MappedKeys[0], UE::Input::TestKeyX);
+	}
+	
+	return true;
+}
\ No newline at end of file
diff --git a/Engine/Plugins/EnhancedInput/Source/InputEditor/Private/Tests/InputTestFramework.cpp b/Engine/Plugins/EnhancedInput/Source/InputEditor/Private/Tests/InputTestFramework.cpp
--- a/Engine/Plugins/EnhancedInput/Source/InputEditor/Private/Tests/InputTestFramework.cpp
+++ b/Engine/Plugins/EnhancedInput/Source/InputEditor/Private/Tests/InputTestFramework.cpp
@@ -8,14 +8,42 @@
 #include "GameFramework/PlayerController.h"
 #include "InputMappingContext.h"
 #include "Tests/AutomationEditorCommon.h"
+#include "UserSettings/EnhancedInputUserSettings.h"
+#include "NativeGameplayTags.h"
 
 #include UE_INLINE_GENERATED_CPP_BY_NAME(InputTestFramework)
 
 FMockedEnhancedInputSubsystem::FMockedEnhancedInputSubsystem(const UControllablePlayer& PlayerData)
 : PlayerInput(PlayerData.PlayerInput)
+, UserSettings(PlayerData.UserSettings)
 {
 }
 
+UEnhancedInputUserSettings* FMockedEnhancedInputSubsystem::GetUserSettings() const
+{
+	return UserSettings.Get();
+}
+
+void FMockedEnhancedInputSubsystem::InitalizeUserSettings()
+{
+	ULocalPlayer* LP = PlayerInput->GetOwningLocalPlayer();
+	
+	UserSettings->Initialize(LP);
+	UserSettings->ApplySettings();
+}
+
+void UMockInputUserSettings::Initialize(ULocalPlayer* LP)
+{
+	// Create a default key mapping profile in the case where one doesn't exist
+	if (!GetCurrentKeyProfile())
+	{
+		FPlayerMappableKeyProfileCreationArgs Args = {};
+		Args.bSetAsCurrentProfile = true;
+	
+		CreateNewKeyProfile(Args);
+	}
+}
+
 UWorld* AnEmptyWorld()
 {
 #if WITH_AUTOMATION_TESTS
@@ -36,8 +64,10 @@ UControllablePlayer& AControllablePlayer(UWorld* World)
 	PlayerData->PlayerInput = Cast<UEnhancedPlayerInput>(PlayerData->Player->PlayerInput);
 	PlayerData->InputComponent = Cast<UEnhancedInputComponent>(PlayerData->Player->InputComponent);
 	PlayerData->Player->InitInputSystem();
+	PlayerData->UserSettings = NewObject<UMockInputUserSettings>(GetTransientPackage(), UMockInputUserSettings::StaticClass());
 
 	PlayerData->Subsystem.Reset(new FMockedEnhancedInputSubsystem(*PlayerData));
+	PlayerData->Subsystem->Init();
 
 	check(PlayerData->IsValid());
 
@@ -53,8 +83,10 @@ UInputMappingContext* AnInputContextIsAppliedToAPlayer(UControllablePlayer& Play
 
 UInputAction* AnInputAction(UControllablePlayer& PlayerData, FName ActionName, EInputActionValueType ValueType)
 {
-	UInputAction* Action = NewObject<UInputAction>(PlayerData.Player, ActionName);
+	UTestMappableKeysAction* Action = NewObject<UTestMappableKeysAction>(PlayerData.Player, ActionName);
 	Action->ValueType = ValueType;
+	Action->SetPlayerMappableKeyOptions(ActionName);
+	
 	return PlayerData.InputAction.Emplace(ActionName, Action);
 }
 
diff --git a/Engine/Plugins/EnhancedInput/Source/InputEditor/Private/Tests/InputTestFramework.h b/Engine/Plugins/EnhancedInput/Source/InputEditor/Private/Tests/InputTestFramework.h
--- a/Engine/Plugins/EnhancedInput/Source/InputEditor/Private/Tests/InputTestFramework.h
+++ b/Engine/Plugins/EnhancedInput/Source/InputEditor/Private/Tests/InputTestFramework.h
@@ -3,6 +3,8 @@
 
 
 #include "EnhancedInputSubsystemInterface.h"
+#include "PlayerMappableKeySettings.h"
+#include "UserSettings/EnhancedInputUserSettings.h"
 #include "InputTestFramework.generated.h"
 
 struct FEnhancedActionKeyMapping;
@@ -65,9 +67,24 @@ struct FBindingTargets
 class FMockedEnhancedInputSubsystem : public IEnhancedInputSubsystemInterface
 {
 	TWeakObjectPtr<UEnhancedPlayerInput> PlayerInput;
+	
+	TWeakObjectPtr<UEnhancedInputUserSettings> UserSettings;
+	
 public:
 	FMockedEnhancedInputSubsystem(const class UControllablePlayer& PlayerData);
+	void Init() { InitalizeUserSettings(); }
 	virtual UEnhancedPlayerInput* GetPlayerInput() const override { return PlayerInput.Get(); }
+	virtual UEnhancedInputUserSettings* GetUserSettings() const override;
+	virtual void InitalizeUserSettings() override;
+};
+
+UCLASS()
+class UMockInputUserSettings : public UEnhancedInputUserSettings
+{
+	GENERATED_BODY()
+public:
+	// There won't be an owning local player here so we will override it to avoid unnecessary error logging.
+	virtual void Initialize(ULocalPlayer* LP) override;
 };
 
 UCLASS()
@@ -98,18 +115,39 @@ public:
 	// Ensure we don't try to double bind listeners when applying multiple key mappings
 	TSet<const UInputAction*> MappedActionListeners;
 
-	bool IsValid() const { return Player && Subsystem && PlayerInput.IsValid() && InputComponent.IsValid(); }
+	// The user settings object that we can use to test. Normally this is owned by the EI Local Player subsystem 
+	UPROPERTY()
+	TObjectPtr<UEnhancedInputUserSettings> UserSettings;
+
+	bool IsValid() const { return Player && Subsystem && PlayerInput.IsValid() && InputComponent.IsValid() && UserSettings; }
 
 	// Mocked subsystem implementing IEnhancedInputSubsystemInterface
 	TUniquePtr<FMockedEnhancedInputSubsystem> Subsystem;
 };
 
+// A subclass of UInputAction that will have it's player mappable key settings object set automatically so 
+// we can test re-mappable keys.
+UCLASS()
+class UTestMappableKeysAction : public UInputAction
+{
+	GENERATED_BODY()
+public:
+
+	void SetPlayerMappableKeyOptions(const FName MappingName)
+	{
+		PlayerMappableKeySettings = NewObject<UPlayerMappableKeySettings>(this, TEXT("TestMappingSettings"));
+		PlayerMappableKeySettings->Name = MappingName;
+	}
+};
+
 // Default test values
 static FName TestContext = TEXT("TestContext");
 static FName TestAction = TEXT("TestAction");
 static FName TestAction2 = TEXT("TestAction2");
 static FName TestAction3 = TEXT("TestAction3");
 static FName TestAction4 = TEXT("TestAction4");
+static FName TestAction5 = TEXT("TestAction5");
+static FName TestAction6 = TEXT("TestAction6");
 // TODO: Create custom keys for the lifetime of the test. Using public keys can cause interference (e.g. AxisConfigs)
 static FKey TestKey = EKeys::A;
 static FKey TestKey2 = EKeys::S;
