Mark the test input action to not show up in a class drop down

#rnx

diff --git a/Engine/Plugins/EnhancedInput/Source/InputEditor/Private/Tests/InputTestFramework.h b/Engine/Plugins/EnhancedInput/Source/InputEditor/Private/Tests/InputTestFramework.h
--- a/Engine/Plugins/EnhancedInput/Source/InputEditor/Private/Tests/InputTestFramework.h
+++ b/Engine/Plugins/EnhancedInput/Source/InputEditor/Private/Tests/InputTestFramework.h
@@ -78,7 +78,7 @@ public:
 	virtual void InitalizeUserSettings() override;
 };
 
-UCLASS()
+UCLASS(hidedropdown, NotBlueprintable, NotBlueprintType)
 class UMockInputUserSettings : public UEnhancedInputUserSettings
 {
 	GENERATED_BODY()
@@ -127,7 +127,7 @@ public:
 
 // A subclass of UInputAction that will have it's player mappable key settings object set automatically so 
 // we can test re-mappable keys.
-UCLASS()
+UCLASS(hidedropdown, NotBlueprintable, NotBlueprintType)
 class UTestMappableKeysAction : public UInputAction
 {
 	GENERATED_BODY()
