Added support for Enhanced Input show debug command to use debug target, fixing issue where show debug menu always only showed input information of first local player in a game instance even if debug target selected others.

diff --git a/Engine/Plugins/EnhancedInput/Source/EnhancedInput/Private/EnhancedInputModule.cpp b/Engine/Plugins/EnhancedInput/Source/EnhancedInput/Private/EnhancedInputModule.cpp
--- a/Engine/Plugins/EnhancedInput/Source/EnhancedInput/Private/EnhancedInputModule.cpp
+++ b/Engine/Plugins/EnhancedInput/Source/EnhancedInput/Private/EnhancedInputModule.cpp
@@ -10,6 +10,7 @@
 #include "EnhancedInputSubsystems.h"
 #include "GameFramework/HUD.h"
 #include "GameFramework/PlayerController.h"
+#include "GameFramework/Pawn.h"
 #include "UObject/Package.h"
 #include "UObject/UObjectIterator.h"
 #include "EnhancedInputDeveloperSettings.h"
@@ -355,12 +356,24 @@ void FEnhancedInputModule::OnShowDebugInfo(AHUD* HUD, UCanvas* Canvas, const FDe
 		// Show the debug info for input user settings
 		if (HUD->ShouldDisplayDebug(NAME_InputSettings))
 		{
-			// TODO: some way to page through different local players instead of
-			// only showing the first player
-			TObjectIterator<UEnhancedInputLocalPlayerSubsystem> FirstPlayer;
-			if (FirstPlayer)
+			const APawn* TargetPawn = Cast<APawn>(HUD->GetCurrentDebugTargetActor());
+			const APlayerController* TargetPlayerController = TargetPawn ? TargetPawn->GetController<APlayerController>() : Cast<APlayerController>(HUD->GetCurrentDebugTargetActor());
+			const ULocalPlayer* TargetLocalPlayer = TargetPlayerController ? TargetPlayerController->GetLocalPlayer() : nullptr;
+			const UEnhancedInputLocalPlayerSubsystem* TargetEnhancedInputSubsystem = TargetPlayerController ? TargetLocalPlayer->GetSubsystem<UEnhancedInputLocalPlayerSubsystem>() : nullptr;
+
+			//Fallback onto first found input system.
+			if (!TargetEnhancedInputSubsystem)
+			{
+				TObjectIterator<UEnhancedInputLocalPlayerSubsystem> FirstFoundInputSubsystem;
+				if (FirstFoundInputSubsystem)
+				{
+					TargetEnhancedInputSubsystem = *FirstFoundInputSubsystem;
+				}
+			}
+			
+			if (TargetEnhancedInputSubsystem)
 			{
-				if (UEnhancedInputUserSettings* Settings = FirstPlayer->GetUserSettings())
+				if (UEnhancedInputUserSettings* Settings = TargetEnhancedInputSubsystem->GetUserSettings())
 				{
 					Settings->ShowDebugInfo(Canvas);
 				}
