Fix a bad include that was automatically added

#rnx

diff --git a/Engine/Plugins/EnhancedInput/Source/InputEditor/Private/Tests/InputPlayerMappableKeysTests.cpp b/Engine/Plugins/EnhancedInput/Source/InputEditor/Private/Tests/InputPlayerMappableKeysTests.cpp
--- a/Engine/Plugins/EnhancedInput/Source/InputEditor/Private/Tests/InputPlayerMappableKeysTests.cpp
+++ b/Engine/Plugins/EnhancedInput/Source/InputEditor/Private/Tests/InputPlayerMappableKeysTests.cpp
@@ -1,7 +1,5 @@
 // Copyright Epic Games, Inc. All Rights Reserved.
 
-#include <ThirdParty/Perforce/p4api-2018.1/Include/Win64/VS2015/p4/error.h>
-
 #include "InputMappingContext.h"
 #include "InputTestFramework.h"
 #include "Misc/AutomationTest.h"
