Add the ability to bind lambda expressions to the Enhanced Input component outright. Also mark the delegate handles as protected instead of private so that users can more easily extend upon the Enhanced Input component and add more binding options to it if they want.

diff --git a/Engine/Plugins/EnhancedInput/Source/EnhancedInput/Public/EnhancedInputComponent.h b/Engine/Plugins/EnhancedInput/Source/EnhancedInput/Public/EnhancedInputComponent.h
--- a/Engine/Plugins/EnhancedInput/Source/EnhancedInput/Public/EnhancedInputComponent.h
+++ b/Engine/Plugins/EnhancedInput/Source/EnhancedInput/Public/EnhancedInputComponent.h
@@ -35,7 +35,7 @@ DECLARE_DYNAMIC_DELEGATE_FourParams(FEnhancedInputActionHandlerDynamicSignature,
 template<typename TSignature>
 struct TEnhancedInputUnifiedDelegate
 {
-private:
+protected:
 	/** Holds the delegate to call. */
 	TSharedPtr<TSignature> Delegate;
 
@@ -86,6 +86,15 @@ public:
 		Delegate->BindUFunction(Object, FuncName);
 	}
 
+	/** Binds a lambda expression */
+	template<typename TSig = TSignature, typename FunctorType, typename... VarTypes>
+	void BindLambda(FunctorType&& InFunctor, VarTypes&&... Vars)
+	{
+		Unbind();
+		Delegate = MakeShared<TSig>();
+		Delegate->BindLambda(Forward<FunctorType>(InFunctor), Forward<VarTypes>(Vars)...);
+	}
+
 	template<typename TSig = TSignature>
 	TSig& MakeDelegate()
 	{
@@ -329,7 +338,7 @@ class ENHANCEDINPUT_API UEnhancedInputComponent
 {
 	GENERATED_UCLASS_BODY()
 
-private:
+protected:
 
 	/** The collection of action bindings. */
 	TArray<TUniquePtr<FEnhancedInputActionEventBinding>> EnhancedActionEventBindings;
@@ -441,6 +450,50 @@ public:
 		return *EnhancedActionEventBindings.Add_GetRef(MoveTemp(AB));
 	}
 
+	/**
+	 * Binds a lambda expression to receive the input action value form an input event.
+	 * Example:
+	 *              InputComponent->BindActionValueLambda(
+	 *					TestAction,
+	 *					ETriggerEvent::Triggered,
+	 *					[](const FInputActionValue& InputActionValue)
+	 *					{
+	 *						// Do you work here!
+	 *					});
+	 */
+	template<typename FunctorType, typename... VarTypes>
+	FEnhancedInputActionEventBinding& BindActionValueLambda(const UInputAction* Action, ETriggerEvent TriggerEvent, FunctorType&& InFunctor, VarTypes&&... Vars)
+	{
+		TUniquePtr<FEnhancedInputActionEventDelegateBinding<FEnhancedInputActionHandlerValueSignature>> AB = MakeUnique<FEnhancedInputActionEventDelegateBinding<FEnhancedInputActionHandlerValueSignature>>(Action, TriggerEvent);
+	
+		AB->Delegate.BindLambda(Forward<FunctorType>(InFunctor), Forward<VarTypes>(Vars)...);
+	
+		AB->Delegate.SetShouldFireWithEditorScriptGuard(bShouldFireDelegatesInEditor);
+		return *EnhancedActionEventBindings.Add_GetRef(MoveTemp(AB));
+	}
+
+	/**
+	 * Bind a lambda to receive the action instance data from an input event.
+	 * Example:
+	 *			InputComponent->BindActionInstanceLambda(
+	 *					TestAction,
+	 *					ETriggerEvent::Triggered,
+	 *					[](const FInputActionInstance& ActionInstance)
+	 *					{
+	 *						// Do you work here!
+	 *					});
+	 */
+	template<typename FunctorType, typename... VarTypes>
+    FEnhancedInputActionEventBinding& BindActionInstanceLambda(const UInputAction* Action, ETriggerEvent TriggerEvent, FunctorType&& InFunctor, VarTypes&&... Vars)
+    {
+    	TUniquePtr<FEnhancedInputActionEventDelegateBinding<FEnhancedInputActionHandlerInstanceSignature>> AB = MakeUnique<FEnhancedInputActionEventDelegateBinding<FEnhancedInputActionHandlerInstanceSignature>>(Action, TriggerEvent);
+    
+    	AB->Delegate.BindLambda(Forward<FunctorType>(InFunctor), Forward<VarTypes>(Vars)...);
+    
+    	AB->Delegate.SetShouldFireWithEditorScriptGuard(bShouldFireDelegatesInEditor);
+    	return *EnhancedActionEventBindings.Add_GetRef(MoveTemp(AB));
+    }
+	
 	/**
 	 * Binds a UInputAction assigned via UInputMappingContext to this component.
 	 * No delegate will be called when this action triggers. The binding simply reflects the current value of the action.
