Update the label on the asset creation tab to mark Player Mappable Input Config types as deprecated

#rnx

diff --git a/Engine/Plugins/EnhancedInput/Source/InputEditor/Private/InputEditorModule.cpp b/Engine/Plugins/EnhancedInput/Source/InputEditor/Private/InputEditorModule.cpp
--- a/Engine/Plugins/EnhancedInput/Source/InputEditor/Private/InputEditorModule.cpp
+++ b/Engine/Plugins/EnhancedInput/Source/InputEditor/Private/InputEditorModule.cpp
@@ -315,7 +315,7 @@ PRAGMA_DISABLE_DEPRECATION_WARNINGS
 class UE_DEPRECATED(5.3, "PlayerMappableInputConfig has been deprecated, please use UEnhancedInputUserSettings instead") FAssetTypeActions_PlayerMappableInputConfig : public FAssetTypeActions_DataAsset
 {
 public:
-	virtual FText GetName() const override { return NSLOCTEXT("AssetTypeActions", "AssetTypeActions_PlayerMappableInputConfig", "Player Mappable Input Config"); }
+	virtual FText GetName() const override { return NSLOCTEXT("AssetTypeActions", "AssetTypeActions_PlayerMappableInputConfig", "Player Mappable Input Config (Deprecated)"); }
 	virtual uint32 GetCategories() override { return FInputEditorModule::GetInputAssetsCategory(); }
 	virtual FColor GetTypeColor() const override { return FColor(127, 255, 255); }
 	virtual FText GetAssetDescription(const FAssetData& AssetData) const override { return NSLOCTEXT("AssetTypeActions", "AssetTypeActions_PlayerBindableInputConfigDesc", "Represents one set of Player Mappable controller/keymappings"); }
