PR #10630: Added new Input Modifier - Smooth Delta.

It provides a normalized smoothed difference between the current and last input value, which options on how the value is calculated for different behavior.

For boolean Input Values, it just returns the boolean as is.

diff --git a/Engine/Plugins/EnhancedInput/Source/EnhancedInput/Private/InputModifiers.cpp b/Engine/Plugins/EnhancedInput/Source/EnhancedInput/Private/InputModifiers.cpp
--- a/Engine/Plugins/EnhancedInput/Source/EnhancedInput/Private/InputModifiers.cpp
+++ b/Engine/Plugins/EnhancedInput/Source/EnhancedInput/Private/InputModifiers.cpp
@@ -13,6 +13,46 @@
 
 #define LOCTEXT_NAMESPACE "EnhancedInputModifiers"
 
+/*
+* Normalized Smooth Delta
+*/
+
+FInputActionValue UInputModifierSmoothDelta::ModifyRaw_Implementation(const UEnhancedPlayerInput* PlayerInput, FInputActionValue CurrentValue, float DeltaTime)
+{
+	// You can't smooth a boolean value
+	if (ensureMsgf(CurrentValue.GetValueType() != EInputActionValueType::Boolean, TEXT("The 'Smooth Delta Modifier' doesn't support boolean values.")))
+	{
+		return CurrentValue;
+	}
+
+	FVector NewValue = CurrentValue.Get<FVector>();
+	FVector TargetDelta = (NewValue - OldValue).GetSafeNormal();
+	OldValue = NewValue;
+
+	switch (SmoothingMethod)
+	{
+		case ENormalizeInputSmoothingType::Lerp: Delta = FMath::LerpStable(Delta, TargetDelta, Speed); break;
+		case ENormalizeInputSmoothingType::Interp_To: Delta = FMath::VInterpTo(Delta, TargetDelta, DeltaTime, Speed); break;
+		case ENormalizeInputSmoothingType::Interp_Constant_To: Delta = FMath::VInterpConstantTo(Delta, TargetDelta, DeltaTime, Speed); break;
+		case ENormalizeInputSmoothingType::Interp_Circular_In: Delta = FMath::InterpCircularIn(Delta, TargetDelta, Speed); break;
+		case ENormalizeInputSmoothingType::Interp_Circular_Out: Delta = FMath::InterpCircularOut(Delta, TargetDelta, Speed); break;
+		case ENormalizeInputSmoothingType::Interp_Circular_In_Out: Delta = FMath::InterpCircularInOut(Delta, TargetDelta, Speed); break;
+		case ENormalizeInputSmoothingType::Interp_Ease_In: Delta = FMath::InterpEaseIn(Delta, TargetDelta, Speed, EasingExponent); break;
+		case ENormalizeInputSmoothingType::Interp_Ease_Out: Delta = FMath::InterpEaseOut(Delta, TargetDelta, Speed, EasingExponent); break;
+		case ENormalizeInputSmoothingType::Interp_Ease_In_Out: Delta = FMath::InterpEaseInOut(Delta, TargetDelta, Speed, EasingExponent); break;
+		case ENormalizeInputSmoothingType::Interp_Expo_In: Delta = FMath::InterpExpoIn(Delta, TargetDelta, Speed); break;
+		case ENormalizeInputSmoothingType::Interp_Expo_Out: Delta = FMath::InterpExpoOut(Delta, TargetDelta, Speed); break;
+		case ENormalizeInputSmoothingType::Interp_Expo_In_Out: Delta = FMath::InterpExpoInOut(Delta, TargetDelta, Speed); break;
+		case ENormalizeInputSmoothingType::Interp_Sin_In: Delta = FMath::InterpSinIn(Delta, TargetDelta, Speed); break;
+		case ENormalizeInputSmoothingType::Interp_Sin_Out: Delta = FMath::InterpSinOut(Delta, TargetDelta, Speed); break;
+		case ENormalizeInputSmoothingType::Interp_Sin_In_Out: Delta = FMath::InterpSinInOut(Delta, TargetDelta, Speed); break;
+		default: Delta = TargetDelta;
+	}
+
+	return Delta;
+}
+
+
 /*
 * Scalar
 */
diff --git a/Engine/Plugins/EnhancedInput/Source/EnhancedInput/Public/InputModifiers.h b/Engine/Plugins/EnhancedInput/Source/EnhancedInput/Public/InputModifiers.h
--- a/Engine/Plugins/EnhancedInput/Source/EnhancedInput/Public/InputModifiers.h
+++ b/Engine/Plugins/EnhancedInput/Source/EnhancedInput/Public/InputModifiers.h
@@ -52,6 +52,65 @@ public:
 	FLinearColor GetVisualizationColor(FInputActionValue SampleValue, FInputActionValue FinalValue) const;
 };
 
+/** Behavior options for how the value is calculated in UInputModifierSmoothDelta. */
+UENUM()
+enum class ENormalizeInputSmoothingType : uint8
+{
+	None UMETA(Hidden),
+	Lerp,
+	Interp_To,
+	Interp_Constant_To,
+	Interp_Circular_In, 
+	Interp_Circular_Out, 
+	Interp_Circular_In_Out,
+	Interp_Ease_In, 
+	Interp_Ease_Out, 
+	Interp_Ease_In_Out,
+	Interp_Expo_In,
+	Interp_Expo_Out,
+	Interp_Expo_In_Out,
+	Interp_Sin_In,
+	Interp_Sin_Out,
+	Interp_Sin_In_Out
+};
+
+/**
+* Normalized Smooth Delta
+* 
+* Produces a smoothed normalized delta of the current(new) and last(old) input value.
+* Boolean input values will be returned as is.
+*/
+UCLASS(NotBlueprintable, MinimalAPI, meta = (DisplayName = "Smooth Delta"))
+class UInputModifierSmoothDelta : public UInputModifier
+{
+	GENERATED_BODY()
+
+public:
+
+	UPROPERTY(EditInstanceOnly, BlueprintReadWrite, Category = Settings, Config)
+	ENormalizeInputSmoothingType SmoothingMethod = ENormalizeInputSmoothingType::Lerp;
+
+	// Speed, or Alpha. If the speed given is 0, then jump to the target.
+	UPROPERTY(EditInstanceOnly, BlueprintReadWrite, Category = Settings, Config)
+	float Speed = 0.5f;
+
+	/**
+	 * For ease functions, this controls the degree of the curve.
+	 * 
+	 * This only affects the Interp_Ease_In, Interp_Ease_Out, and Interp_Ease_In_Out smoothing methods
+	 */
+	UPROPERTY(EditInstanceOnly, BlueprintReadWrite, Category = Settings, Config, meta=(EditCondition="SmoothingMethod == ENormalizeInputSmoothingType::Interp_Ease_In ||SmoothingMethod == ENormalizeInputSmoothingType::Interp_Ease_Out || SmoothingMethod == ENormalizeInputSmoothingType::Interp_Ease_In_Out"))
+	float EasingExponent = 2.0f;
+
+protected:
+	virtual FInputActionValue ModifyRaw_Implementation(const UEnhancedPlayerInput* PlayerInput, FInputActionValue CurrentValue, float DeltaTime) override;
+	
+	FVector OldValue = FVector::ZeroVector;
+    
+    FVector Delta = FVector::ZeroVector;
+};
+
+
 UENUM()
 enum class EDeadZoneType : uint8
 {
