Use the new "GetOptions" metadata on BP callable functions that take in an FName for a "MappingName". This will make referencing Player Mappable settings in Blueprints much easier!

diff --git a/Engine/Plugins/EnhancedInput/Source/EnhancedInput/Public/EnhancedInputSubsystemInterface.h b/Engine/Plugins/EnhancedInput/Source/EnhancedInput/Public/EnhancedInputSubsystemInterface.h
--- a/Engine/Plugins/EnhancedInput/Source/EnhancedInput/Public/EnhancedInputSubsystemInterface.h
+++ b/Engine/Plugins/EnhancedInput/Source/EnhancedInput/Public/EnhancedInputSubsystemInterface.h
@@ -138,7 +138,7 @@ public:
 	 * @param Triggers			The triggers to apply to the injected input.
 	 */
 	UFUNCTION(BlueprintCallable, Category="Input", meta=(AutoCreateRefTerm="Modifiers,Triggers"))
-	virtual void InjectInputForPlayerMapping(const FName MappingName, FInputActionValue RawValue, const TArray<UInputModifier*>& Modifiers, const TArray<UInputTrigger*>& Triggers);
+	virtual void InjectInputForPlayerMapping(UPARAM(Meta=(GetOptions="EnhancedInput.PlayerMappableKeySettings.GetKnownMappingNames")) const FName MappingName, FInputActionValue RawValue, const TArray<UInputModifier*>& Modifiers, const TArray<UInputTrigger*>& Triggers);
 	
 	/**
 	 * Input simulation via injection. Runs modifiers and triggers delegates as if the input had come through the underlying input system as FKeys.
@@ -150,7 +150,7 @@ public:
 	 * @param Triggers			The triggers to apply to the injected input.
 	 */
 	UFUNCTION(BlueprintCallable, Category="Input", meta=(AutoCreateRefTerm="Modifiers,Triggers"))
-	virtual void InjectInputVectorForPlayerMapping(const FName MappingName, FVector Value, const TArray<UInputModifier*>& Modifiers, const TArray<UInputTrigger*>& Triggers);
+	virtual void InjectInputVectorForPlayerMapping(UPARAM(Meta=(GetOptions="EnhancedInput.PlayerMappableKeySettings.GetKnownMappingNames")) const FName MappingName, FVector Value, const TArray<UInputModifier*>& Modifiers, const TArray<UInputTrigger*>& Triggers);
 
 	/**
 	 * Starts simulation of input via injection. This injects the given input every tick until it is stopped with StopContinuousInputInjectionForAction.
@@ -172,7 +172,7 @@ public:
 	 * @param Triggers			The triggers to apply to the injected input.
 	 */
 	UFUNCTION(BlueprintCallable, Category="Input", meta=(AutoCreateRefTerm="Modifiers,Triggers"))
-	virtual void StartContinuousInputInjectionForPlayerMapping(const FName MappingName, FInputActionValue RawValue, const TArray<UInputModifier*>& Modifiers, const TArray<UInputTrigger*>& Triggers);
+	virtual void StartContinuousInputInjectionForPlayerMapping(UPARAM(Meta=(GetOptions="EnhancedInput.PlayerMappableKeySettings.GetKnownMappingNames")) const FName MappingName, FInputActionValue RawValue, const TArray<UInputModifier*>& Modifiers, const TArray<UInputTrigger*>& Triggers);
 
 	/**
 	 * Stops continuous input injection for the given action.
@@ -188,7 +188,7 @@ public:
 	 * @param MappingName		The name of the player mapping that can be used for look up an associated UInputAction object.
 	 */
 	UFUNCTION(BlueprintCallable, Category="Input")
-	virtual void StopContinuousInputInjectionForPlayerMapping(const FName MappingName);
+	virtual void StopContinuousInputInjectionForPlayerMapping(UPARAM(Meta=(GetOptions="EnhancedInput.PlayerMappableKeySettings.GetKnownMappingNames")) const FName MappingName);
 	
 	/**
 	 * Remove all applied mapping contexts.
diff --git a/Engine/Plugins/EnhancedInput/Source/EnhancedInput/Public/UserSettings/EnhancedInputUserSettings.h b/Engine/Plugins/EnhancedInput/Source/EnhancedInput/Public/UserSettings/EnhancedInputUserSettings.h
--- a/Engine/Plugins/EnhancedInput/Source/EnhancedInput/Public/UserSettings/EnhancedInputUserSettings.h
+++ b/Engine/Plugins/EnhancedInput/Source/EnhancedInput/Public/UserSettings/EnhancedInputUserSettings.h
@@ -316,7 +316,7 @@ public:
 
 	/** Resets every player key mapping to this mapping back to it's default value */
 	UFUNCTION(BlueprintCallable, Category="Enhanced Input|User Settings")
-	void ResetMappingToDefault(const FName InMappingName);
+	void ResetMappingToDefault(UPARAM(Meta=(GetOptions="EnhancedInput.PlayerMappableKeySettings.GetKnownMappingNames")) const FName InMappingName);
 	
 	/** Get all the key mappings associated with the given mapping name on this profile */
 	FKeyMappingRow* FindKeyMappingRowMutable(const FName InMappingName);
@@ -338,7 +338,7 @@ public:
 	 * Returns the number of keys for the given mapping name
 	 */
 	UFUNCTION(BlueprintCallable, Category="Enhanced Input|User Settings", meta = (ReturnDisplayName = "Number of keys"))
-	virtual int32 GetMappedKeysInRow(const FName MappingName, /*OUT*/ TArray<FKey>& OutKeys) const;
+	virtual int32 GetMappedKeysInRow(UPARAM(Meta=(GetOptions="EnhancedInput.PlayerMappableKeySettings.GetKnownMappingNames")) const FName MappingName, /*OUT*/ TArray<FKey>& OutKeys) const;
 
 	/**
 	 * Called during IEnhancedInputSubsystemInterface::RebuildControlMappings, this provides your user settings
@@ -591,7 +591,7 @@ public:
 
 	/** Returns a set of all player key mappings for the given mapping name. */
 	UFUNCTION(BlueprintCallable, Category="Enhanced Input|User Settings")
-	virtual const TSet<FPlayerKeyMapping>& FindMappingsInRow(const FName MappingName) const;
+	virtual const TSet<FPlayerKeyMapping>& FindMappingsInRow(UPARAM(Meta=(GetOptions="EnhancedInput.PlayerMappableKeySettings.GetKnownMappingNames")) const FName MappingName) const;
 
 	/** Returns the current player key mapping for the given row name in the given slot */
 	virtual const FPlayerKeyMapping* FindCurrentMappingForSlot(const FName MappingName, const EPlayerMappableKeySlot InSlot) const;
