PR #10879: Make UInputModifier::ModifyRaw non-const

diff --git a/Engine/Plugins/EnhancedInput/Source/EnhancedInput/Public/InputModifiers.h b/Engine/Plugins/EnhancedInput/Source/EnhancedInput/Public/InputModifiers.h
--- a/Engine/Plugins/EnhancedInput/Source/EnhancedInput/Public/InputModifiers.h
+++ b/Engine/Plugins/EnhancedInput/Source/EnhancedInput/Public/InputModifiers.h
@@ -41,7 +41,7 @@ public:
 	 * @param CurrentValue - The modified value returned by the previous modifier in the chain, or the base raw value if this is the first modifier in the chain.
 	 */
 	UFUNCTION(BlueprintCallable, BlueprintNativeEvent, Category = "Modifier")
-	FInputActionValue ModifyRaw(const UEnhancedPlayerInput* PlayerInput, FInputActionValue CurrentValue, float DeltaTime) const;
+	FInputActionValue ModifyRaw(const UEnhancedPlayerInput* PlayerInput, FInputActionValue CurrentValue, float DeltaTime);
 
 	/**
 	 * Helper to allow debug visualization of the modifier.
