[Backout] - CL27999507

Backing this out because it caused any existing BP modifier implementations to have cook errors. Instead of removing the const we would need to properly deprecate it and make a new version to support old content

Original CL Desc
-----------------------------------------------------------------
PR #10879: Make UInputModifier::ModifyRaw non-const

diff --git a/Engine/Plugins/EnhancedInput/Source/EnhancedInput/Public/InputModifiers.h b/Engine/Plugins/EnhancedInput/Source/EnhancedInput/Public/InputModifiers.h
--- a/Engine/Plugins/EnhancedInput/Source/EnhancedInput/Public/InputModifiers.h
+++ b/Engine/Plugins/EnhancedInput/Source/EnhancedInput/Public/InputModifiers.h
@@ -41,7 +41,7 @@ public:
 	 * @param CurrentValue - The modified value returned by the previous modifier in the chain, or the base raw value if this is the first modifier in the chain.
 	 */
 	UFUNCTION(BlueprintCallable, BlueprintNativeEvent, Category = "Modifier")
-	FInputActionValue ModifyRaw(const UEnhancedPlayerInput* PlayerInput, FInputActionValue CurrentValue, float DeltaTime);
+	FInputActionValue ModifyRaw(const UEnhancedPlayerInput* PlayerInput, FInputActionValue CurrentValue, float DeltaTime) const;
 
 	/**
 	 * Helper to allow debug visualization of the modifier.
