Add a "UpdateValueOfContinuousInputInjection" function to allow for you to update the value of a continuously injected input action and preserving the trigger/modifier state. This effectivly allows for modifier and trigger support via input injection because the EI subsystem will keep the state of those modifiers.

The use case for this is mobile UI where you want to use an input trigger for something like a gesture and you want the input trigger to be created once for a thumbstick and a touch screen.

Copying the values of triggers and modifiers in the ContinousInjection so that we can evaluate a new state instead of using the old one

diff --git a/Engine/Plugins/EnhancedInput/Source/EnhancedInput/Private/EnhancedInputSubsystemInterface.cpp b/Engine/Plugins/EnhancedInput/Source/EnhancedInput/Private/EnhancedInputSubsystemInterface.cpp
--- a/Engine/Plugins/EnhancedInput/Source/EnhancedInput/Private/EnhancedInputSubsystemInterface.cpp
+++ b/Engine/Plugins/EnhancedInput/Source/EnhancedInput/Private/EnhancedInputSubsystemInterface.cpp
@@ -30,6 +30,19 @@ static FAutoConsoleVariableRef GCVarGlobalAxisConfigMode(
 	TEXT("Whether or not to apply Global Axis Config settings. 0 = Default (Mouse Only), 1 = All, 2 = None")
 );
 
+template<typename T>
+void DeepCopyPtrArray(const TArray<T*>& From, TArray<T*>& To)
+{
+	To.Empty(From.Num());
+	for (T* ToDuplicate : From)
+	{
+		if (ToDuplicate)
+		{
+			To.Add(DuplicateObject<T>(ToDuplicate, nullptr));
+		}
+	}
+}
+
 void IEnhancedInputSubsystemInterface::InitalizeUserSettings()
 {
 	// Not every implementer of the EI subsystem wants user settings, so leave it up to them to determine if they want it or not
@@ -98,8 +111,8 @@ void IEnhancedInputSubsystemInterface::StartContinuousInputInjectionForAction(co
 	FInjectedInput& Injection = ContinuouslyInjectedInputs.FindOrAdd(Action);
 	
 	Injection.RawValue = RawValue;
-	Injection.Modifiers = Modifiers;
-	Injection.Triggers = Triggers;
+	DeepCopyPtrArray<UInputModifier>(Modifiers, Injection.Modifiers);
+	DeepCopyPtrArray<UInputTrigger>(Triggers, Injection.Triggers);	
 }
 
 void IEnhancedInputSubsystemInterface::StartContinuousInputInjectionForPlayerMapping(const FName MappingName, FInputActionValue RawValue, const TArray<UInputModifier*>& Modifiers, const TArray<UInputTrigger*>& Triggers)
@@ -121,6 +134,33 @@ void IEnhancedInputSubsystemInterface::StartContinuousInputInjectionForPlayerMap
 	}
 }
 
+void IEnhancedInputSubsystemInterface::UpdateValueOfContinuousInputInjectionForAction(const UInputAction* Action, FInputActionValue RawValue)
+{
+	FInjectedInput& Injection = ContinuouslyInjectedInputs.FindOrAdd(Action);
+	Injection.RawValue = RawValue;
+
+	// Do NOT update the triggers/modifiers here to preserve their state
+}
+
+void IEnhancedInputSubsystemInterface::UpdateValueOfContinuousInputInjectionForPlayerMapping(const FName MappingName, FInputActionValue RawValue)
+{
+	if (const UEnhancedInputUserSettings* UserSettings = GetUserSettings())
+	{
+		if (const UInputAction* Action = UserSettings->FindInputActionForMapping(MappingName))
+		{
+			UpdateValueOfContinuousInputInjectionForAction(Action, RawValue);
+		}
+		else
+		{
+			UE_LOG(LogEnhancedInput, Warning, TEXT("Could not find a Input Action for mapping name '%s'"), *MappingName.ToString());
+		}
+	}
+	else
+	{
+		UE_LOG(LogEnhancedInput, Warning, TEXT("Could not find a valid UEnhancedInputUserSettings object, is it enabled in the project settings?"));
+	}
+}
+
 void IEnhancedInputSubsystemInterface::StopContinuousInputInjectionForAction(const UInputAction* Action)
 {
 	ContinuouslyInjectedInputs.Remove(Action);
@@ -725,19 +765,6 @@ void IEnhancedInputSubsystemInterface::RemovePlayerMappableConfig(const UPlayerM
 
 PRAGMA_ENABLE_DEPRECATION_WARNINGS
 
-template<typename T>
-void DeepCopyPtrArray(const TArray<T*>& From, TArray<T*>& To)
-{
-	To.Empty(From.Num());
-	for (T* ToDuplicate : From)
-	{
-		if (ToDuplicate)
-		{
-			To.Add(DuplicateObject<T>(ToDuplicate, nullptr));
-		}
-	}
-}
-
 // TODO: This should be a delegate (along with InjectChordBlockers), moving chording out of the underlying subsystem and enabling implementation of custom mapping handlers.
 /**
  * Reorder the given UnordedMappings such that chording mappings > chorded mappings > everything else.
diff --git a/Engine/Plugins/EnhancedInput/Source/EnhancedInput/Public/EnhancedInputSubsystemInterface.h b/Engine/Plugins/EnhancedInput/Source/EnhancedInput/Public/EnhancedInputSubsystemInterface.h
--- a/Engine/Plugins/EnhancedInput/Source/EnhancedInput/Public/EnhancedInputSubsystemInterface.h
+++ b/Engine/Plugins/EnhancedInput/Source/EnhancedInput/Public/EnhancedInputSubsystemInterface.h
@@ -174,6 +174,24 @@ public:
 	UFUNCTION(BlueprintCallable, Category="Input", meta=(AutoCreateRefTerm="Modifiers,Triggers"))
 	virtual void StartContinuousInputInjectionForPlayerMapping(UPARAM(Meta=(GetOptions="EnhancedInput.PlayerMappableKeySettings.GetKnownMappingNames")) const FName MappingName, FInputActionValue RawValue, const TArray<UInputModifier*>& Modifiers, const TArray<UInputTrigger*>& Triggers);
 
+	/**
+	 * Update the value of a continuous input injection, preserving the state of triggers and modifiers.
+	 *
+	 * @param Action	The Input Action to set inject input for
+	 * @param RawValue	The value to set the action to (the type will be controlled by the Action)
+	 */
+	UFUNCTION(BlueprintCallable, Category="Input")
+	virtual void UpdateValueOfContinuousInputInjectionForAction(const UInputAction* Action, FInputActionValue RawValue);
+
+	/**
+	 * Update the value of a continuous input injection for the given player mapping name, preserving the state of triggers and modifiers.
+	 *
+	 * @param MappingName	The name of the player mapping that can be used for look up an associated UInputAction object.
+	 * @param RawValue		The value to set the action to (the type will be controlled by the Action)
+	 */
+	UFUNCTION(BlueprintCallable, Category="Input")
+	virtual void UpdateValueOfContinuousInputInjectionForPlayerMapping(UPARAM(Meta=(GetOptions="EnhancedInput.PlayerMappableKeySettings.GetKnownMappingNames")) const FName MappingName, FInputActionValue RawValue);
+
 	/**
 	 * Stops continuous input injection for the given action.
 	 *
