Make the Enhanced Input BP node module PreDefault instead of Default. This fixes a bug where enhanced input nodes would sometimes not be loaded in template C++ projects.

diff --git a/Engine/Plugins/EnhancedInput/EnhancedInput.uplugin b/Engine/Plugins/EnhancedInput/EnhancedInput.uplugin
--- a/Engine/Plugins/EnhancedInput/EnhancedInput.uplugin
+++ b/Engine/Plugins/EnhancedInput/EnhancedInput.uplugin
@@ -23,7 +23,7 @@
 		{
 			"Name": "InputBlueprintNodes",
 			"Type": "UncookedOnly",
-			"LoadingPhase": "Default"
+			"LoadingPhase": "PreDefault"
 		},
 		{
 			"Name": "InputEditor",
