Remove TODO comment that doesnt need to be there and provide some more info about the actual property

#rnx

diff --git a/Engine/Plugins/EnhancedInput/Source/EnhancedInput/Public/InputModifiers.h b/Engine/Plugins/EnhancedInput/Source/EnhancedInput/Public/InputModifiers.h
--- a/Engine/Plugins/EnhancedInput/Source/EnhancedInput/Public/InputModifiers.h
+++ b/Engine/Plugins/EnhancedInput/Source/EnhancedInput/Public/InputModifiers.h
@@ -164,7 +164,13 @@ public:
 	virtual EDataValidationResult IsDataValid(class FDataValidationContext& Context) const override;
 #endif
 
-	// TODO: Detail customization to only show modifiable axes for the relevant binding? This thing has no idea what it's bound to...
+	/**
+	 * The scalar that will be applied to the input value.
+	 *
+	 * For example, if you have a scalar of (2.0, 2.0, 2.0), each input axis will be multiplied by 2.0.
+	 *  
+	 * Note: This will do nothing on boolean input action types, as they can only be true or false.
+	 */
 	UPROPERTY(EditInstanceOnly, BlueprintReadWrite, Category=Settings)
 	FVector Scalar = FVector::OneVector;
 
