Add a flag on Default Mapping Contexts for Enhanced Input to register with the User Input settings by default.

diff --git a/Engine/Plugins/EnhancedInput/Source/EnhancedInput/Private/EnhancedInputSubsystems.cpp b/Engine/Plugins/EnhancedInput/Source/EnhancedInput/Private/EnhancedInputSubsystems.cpp
--- a/Engine/Plugins/EnhancedInput/Source/EnhancedInput/Private/EnhancedInputSubsystems.cpp
+++ b/Engine/Plugins/EnhancedInput/Source/EnhancedInput/Private/EnhancedInputSubsystems.cpp
@@ -280,9 +280,21 @@ void UEnhancedInputWorldSubsystem::AddDefaultMappingContexts()
 		{
 			if (const UInputMappingContext* IMC = ContextSetting.InputMappingContext.LoadSynchronous())
 			{
-				if (!HasMappingContext(IMC))
+				if (ContextSetting.bAddImmediately)
 				{
-					AddMappingContext(IMC, ContextSetting.Priority);
+					if (!HasMappingContext(IMC))
+					{
+						FModifyContextOptions Opts = {};
+						Opts.bNotifyUserSettings = ContextSetting.bRegisterWithUserSettings;
+						AddMappingContext(IMC, ContextSetting.Priority, Opts);	
+					}					
+				}
+				else if (ContextSetting.bRegisterWithUserSettings)
+				{
+					if (UEnhancedInputUserSettings* Settings = GetUserSettings())
+					{
+						Settings->RegisterInputMappingContext(IMC);
+					}
 				}
 			}
 		}
@@ -297,7 +309,9 @@ void UEnhancedInputWorldSubsystem::RemoveDefaultMappingContexts()
 		{
 			if (const UInputMappingContext* IMC = ContextSetting.InputMappingContext.LoadSynchronous())
 			{
-				RemoveMappingContext(IMC);
+				FModifyContextOptions Opts = {};
+				Opts.bNotifyUserSettings = ContextSetting.bRegisterWithUserSettings;
+				RemoveMappingContext(IMC, Opts);
 			}
 		}
 	}
diff --git a/Engine/Plugins/EnhancedInput/Source/EnhancedInput/Public/EnhancedInputDeveloperSettings.h b/Engine/Plugins/EnhancedInput/Source/EnhancedInput/Public/EnhancedInputDeveloperSettings.h
--- a/Engine/Plugins/EnhancedInput/Source/EnhancedInput/Public/EnhancedInputDeveloperSettings.h
+++ b/Engine/Plugins/EnhancedInput/Source/EnhancedInput/Public/EnhancedInputDeveloperSettings.h
@@ -25,6 +25,14 @@ struct FDefaultContextSetting
 	/** The prioirty that should be given to this mapping context when it is added */
 	UPROPERTY(EditAnywhere, Config, Category = "Input")
 	int32 Priority = 0;
+
+	/** If true, then this IMC will be applied immediately when the EI subsystem is ready */
+	UPROPERTY(EditAnywhere, Config, Category = "Input")
+	bool bAddImmediately = true;
+	
+	/** If true, then this IMC will be registered with the User Input Settings (if one is available) immediately when the Enhanced Input subsystem starts. */
+	UPROPERTY(EditAnywhere, Config, Category = "Input")
+	bool bRegisterWithUserSettings = false;
 };
 
 /** Developer settings for Enhanced Input */
diff --git a/Engine/Plugins/EnhancedInput/Source/InputEditor/Private/EnhancedInputEditorSubsystem.cpp b/Engine/Plugins/EnhancedInput/Source/InputEditor/Private/EnhancedInputEditorSubsystem.cpp
--- a/Engine/Plugins/EnhancedInput/Source/InputEditor/Private/EnhancedInputEditorSubsystem.cpp
+++ b/Engine/Plugins/EnhancedInput/Source/InputEditor/Private/EnhancedInputEditorSubsystem.cpp
@@ -11,6 +11,7 @@
 #include "EnhancedInputEditorProcessor.h"
 #include "EnhancedInputEditorSettings.h"
 #include "Framework/Application/SlateApplication.h"
+#include "UserSettings/EnhancedInputUserSettings.h"
 
 #include UE_INLINE_GENERATED_CPP_BY_NAME(EnhancedInputEditorSubsystem)
 
@@ -255,9 +256,21 @@ void UEnhancedInputEditorSubsystem::AddDefaultMappingContexts()
 	{
 		if (const UInputMappingContext* IMC = ContextSetting.InputMappingContext.LoadSynchronous())
 		{
-			if (!HasMappingContext(IMC))
+			if (ContextSetting.bAddImmediately)
 			{
-				AddMappingContext(IMC, ContextSetting.Priority);	
+				if (!HasMappingContext(IMC))
+				{
+					FModifyContextOptions Opts = {};
+					Opts.bNotifyUserSettings = ContextSetting.bRegisterWithUserSettings;
+					AddMappingContext(IMC, ContextSetting.Priority, Opts);	
+				}
+			}
+			else if (ContextSetting.bRegisterWithUserSettings)
+			{
+				if (UEnhancedInputUserSettings* Settings = GetUserSettings())
+				{
+					Settings->RegisterInputMappingContext(IMC);
+				}
 			}
 		}
 	}
@@ -269,7 +282,9 @@ void UEnhancedInputEditorSubsystem::RemoveDefaultMappingContexts()
 	{
 		if (const UInputMappingContext* IMC = ContextSetting.InputMappingContext.LoadSynchronous())
 		{
-			RemoveMappingContext(IMC);
+			FModifyContextOptions Opts = {};
+			Opts.bNotifyUserSettings = ContextSetting.bRegisterWithUserSettings;
+			RemoveMappingContext(IMC, Opts);
 		}
 	}
 }
