Mark the SavedKeyProfiles property as "Transient" instead of "SaveGame" to avoid it attempted to be serialized by the USaveGame base class. We manually serialize this in the UEnhancedInputUserSettings::Serialize function instead. This resolves some warning logs that would come up only on the first PIE of attempting to load a settings save game file.

#rnx

diff --git a/Engine/Plugins/EnhancedInput/Source/EnhancedInput/Public/UserSettings/EnhancedInputUserSettings.h b/Engine/Plugins/EnhancedInput/Source/EnhancedInput/Public/UserSettings/EnhancedInputUserSettings.h
--- a/Engine/Plugins/EnhancedInput/Source/EnhancedInput/Public/UserSettings/EnhancedInputUserSettings.h
+++ b/Engine/Plugins/EnhancedInput/Source/EnhancedInput/Public/UserSettings/EnhancedInputUserSettings.h
@@ -719,7 +719,7 @@ protected:
 	/**
 	 * All of the known Key Profiles for this user, including the currently active profile.
 	 */
-	UPROPERTY(SaveGame)
+	UPROPERTY(Transient)
 	TMap<FGameplayTag, TObjectPtr<UEnhancedPlayerMappableKeyProfile>> SavedKeyProfiles;
 	
 	/** The owning Local Player object of these settings */
