Adding input modifiers and triggers to root set when injecting continuous inputs to prevent them from getting GCed.
This was resulting in crashes in Android and PIE with touch input simulation due to touch input modifiers becoming unreachable / deleted while they were still in use.

diff --git a/Engine/Plugins/EnhancedInput/Source/EnhancedInput/Private/EnhancedInputSubsystemInterface.cpp b/Engine/Plugins/EnhancedInput/Source/EnhancedInput/Private/EnhancedInputSubsystemInterface.cpp
--- a/Engine/Plugins/EnhancedInput/Source/EnhancedInput/Private/EnhancedInputSubsystemInterface.cpp
+++ b/Engine/Plugins/EnhancedInput/Source/EnhancedInput/Private/EnhancedInputSubsystemInterface.cpp
@@ -109,10 +109,24 @@ void IEnhancedInputSubsystemInterface::InjectInputForPlayerMapping(const FName M
 void IEnhancedInputSubsystemInterface::StartContinuousInputInjectionForAction(const UInputAction* Action, FInputActionValue RawValue, const TArray<UInputModifier*>& Modifiers, const TArray<UInputTrigger*>& Triggers)
 {
 	FInjectedInput& Injection = ContinuouslyInjectedInputs.FindOrAdd(Action);
-	
+
 	Injection.RawValue = RawValue;
 	DeepCopyPtrArray<UInputModifier>(Modifiers, Injection.Modifiers);
-	DeepCopyPtrArray<UInputTrigger>(Triggers, Injection.Triggers);	
+	DeepCopyPtrArray<UInputTrigger>(Triggers, Injection.Triggers);
+
+	// ContinuouslyInjectedInputs Map is not managed.
+	// Continuous input injections seem to be getting garbage collected and
+	// crashing in UObject::ProcessEvent when calling ModifyRaw.
+	// Band-aid fix: Adding all these to root set to avoid garbage collection.
+	for (UInputModifier* Modifier : Injection.Modifiers)
+	{
+		Modifier->AddToRoot();
+	}
+
+	for (UInputTrigger* Trigger : Injection.Triggers)
+	{
+		Trigger->AddToRoot();
+	}
 }
 
 void IEnhancedInputSubsystemInterface::StartContinuousInputInjectionForPlayerMapping(const FName MappingName, FInputActionValue RawValue, const TArray<UInputModifier*>& Modifiers, const TArray<UInputTrigger*>& Triggers)
@@ -163,6 +177,21 @@ void IEnhancedInputSubsystemInterface::UpdateValueOfContinuousInputInjectionForP
 
 void IEnhancedInputSubsystemInterface::StopContinuousInputInjectionForAction(const UInputAction* Action)
 {
+	// ContinuouslyInjectedInputs Map is not managed.
+	// Continuous input injections seem to be getting garbage collected and
+	// crashing in UObject::ProcessEvent when calling ModifyRaw.
+	// Band-aid fix: Adding all these to root set to avoid garbage collection.
+	if (FInjectedInput* Injection = ContinuouslyInjectedInputs.Find(Action))
+	{
+		for (UInputModifier* Modifier : Injection->Modifiers)
+		{
+			Modifier->RemoveFromRoot();
+		}
+		for (UInputTrigger* Trigger : Injection->Triggers)
+		{
+			Trigger->RemoveFromRoot();
+		}
+	}
 	ContinuouslyInjectedInputs.Remove(Action);
 }
 
diff --git a/Engine/Plugins/EnhancedInput/Source/EnhancedInput/Public/EnhancedPlayerInput.h b/Engine/Plugins/EnhancedInput/Source/EnhancedInput/Public/EnhancedPlayerInput.h
--- a/Engine/Plugins/EnhancedInput/Source/EnhancedInput/Public/EnhancedPlayerInput.h
+++ b/Engine/Plugins/EnhancedInput/Source/EnhancedInput/Public/EnhancedPlayerInput.h
@@ -19,10 +19,21 @@ enum class ETriggerEventInternal : uint8;
 enum class EKeyEvent : uint8;
 class UInputMappingContext;
 
+// ContinuouslyInjectedInputs Map is not managed.
+// Continuous input injections seem to be getting garbage collected and
+// crashing in UObject::ProcessEvent when calling ModifyRaw.
+// Band-aid fix: Making these managed references. Also check modifications to
+// IEnhancedInputSubsystemInterface::Start/StopContinuousInputInjectionForAction.
+USTRUCT()
 struct FInjectedInput
 {
+	GENERATED_BODY()
+
 	FInputActionValue RawValue;
+
+	UPROPERTY(Transient)
 	TArray<UInputTrigger*> Triggers;
+	UPROPERTY(Transient)
 	TArray<UInputModifier*> Modifiers;
 };
 
