Add a cvar to check to enable checking for keys mapped to EKeys::Invalid.

It isn't technically incorrect for an IMC to contain a mapping to the invalid key, but it can be hard to track down which mapping is empty if you are working at a large scale. This check will hopefully make it easier to find.

diff --git a/Engine/Plugins/EnhancedInput/Source/EnhancedInput/Private/EnhancedActionKeyMapping.cpp b/Engine/Plugins/EnhancedInput/Source/EnhancedInput/Private/EnhancedActionKeyMapping.cpp
--- a/Engine/Plugins/EnhancedInput/Source/EnhancedInput/Private/EnhancedActionKeyMapping.cpp
+++ b/Engine/Plugins/EnhancedInput/Source/EnhancedInput/Private/EnhancedActionKeyMapping.cpp
@@ -2,8 +2,9 @@
 
 #include "EnhancedActionKeyMapping.h"
 #include "PlayerMappableKeySettings.h"
-
+#include "HAL/IConsoleManager.h"
 #if WITH_EDITOR
+#include "Internationalization/Text.h"
 #include "Misc/DataValidation.h"
 #endif
 
@@ -11,6 +12,19 @@
 
 #define LOCTEXT_NAMESPACE "ActionKeyMapping"
 
+#if WITH_EDITOR
+namespace UE::EnhancedInput
+{
+	static TAutoConsoleVariable<bool> CVarCheckForEmptyKeyMappingsDuringValidation
+	(
+		TEXT("EnhancedInput.Mappings.bCheckForEmptyKeyMappingsDuringValidation"),
+		false,
+		TEXT("When true, Enhanced Input key mappings will throw an error during validation if they are mapped to an empty key."),
+		ECVF_Default
+	);
+}
+#endif	// WITH_EDITOR
+
 UPlayerMappableKeySettings* FEnhancedActionKeyMapping::GetPlayerMappableKeySettings() const
 {
 	switch (SettingBehavior)
@@ -102,6 +116,24 @@ EDataValidationResult FEnhancedActionKeyMapping::IsDataValid(FDataValidationCont
 		Result = CombineDataValidationResults(Result, PlayerMappableKeySettings->IsDataValid(Context));
 	}
 
+	// Check for an invalid key mapping if enabled. Having an invalid key mapping is not technically disallowed by Enhanced Input, but 
+	// can oftentimes having an invalid key mapping could be something that you want to track down because of some custom system you have
+	// made for your game. This setting makes it really easy to track down any empty mappings! 
+	if (UE::EnhancedInput::CVarCheckForEmptyKeyMappingsDuringValidation.GetValueOnAnyThread() && Key == EKeys::Invalid)
+	{
+		Result = EDataValidationResult::Invalid;
+		
+		FFormatNamedArguments Args;
+		Args.Add(TEXT("ActionName"), FText::FromString(GetNameSafe(Action)));
+
+		const FText ErrorMsg =
+			FText::Format(
+				LOCTEXT("MappedToInvalidKey", "Mapping to Input Action {ActionName}. \nEnhanced Input Key Mappings cannot be mapped to EKeys::Invalid when 'EnhancedInput.Mappings.bCheckForEmptyKeyMappingsDuringValidation' is true!"),
+				Args);
+		
+		Context.AddError(ErrorMsg);
+	}
+
 	// Validate the triggers.
 	bool bContextContainsComboTrigger = false;
 	bool bContextContainsNonComboTrigger = false;
