Fix use-after-free issue. The "ConsumeKey" function can re-allocate the KeyStateMap because of the call to FindOrAdd, which would in turn invalidate the pointers in the "KeysWithEvents" map. Instead of calling consume key directly, we can just do a non-mutating call to "Find" instead.

#rnx
#tests Running the game with address sanitizer on
#localization none

diff --git a/Engine/Plugins/EnhancedInput/Source/EnhancedInput/Private/EnhancedPlayerInput.cpp b/Engine/Plugins/EnhancedInput/Source/EnhancedInput/Private/EnhancedPlayerInput.cpp
--- a/Engine/Plugins/EnhancedInput/Source/EnhancedInput/Private/EnhancedPlayerInput.cpp
+++ b/Engine/Plugins/EnhancedInput/Source/EnhancedInput/Private/EnhancedPlayerInput.cpp
@@ -582,7 +582,10 @@ void UEnhancedPlayerInput::EvaluateInputDelegates(const TArray<UInputComponent*>
 						// Consume all keys that are mapped to this input action with the proper trigger values
 						for (const FKey& KeyToConsume : ConsumptionData->KeysToConsume)
 						{
-							ConsumeKey(KeyToConsume);
+							if (FKeyState* KeyState = KeyStateMap.Find(KeyToConsume))
+							{
+								KeyState->bConsumed = true;
+							}
 						}
 					}
 				}
