Remove some spammy logs of the Ei world subsystem

#rnx

diff --git a/Engine/Plugins/EnhancedInput/Source/EnhancedInput/Private/EnhancedInputSubsystems.cpp b/Engine/Plugins/EnhancedInput/Source/EnhancedInput/Private/EnhancedInputSubsystems.cpp
--- a/Engine/Plugins/EnhancedInput/Source/EnhancedInput/Private/EnhancedInputSubsystems.cpp
+++ b/Engine/Plugins/EnhancedInput/Source/EnhancedInput/Private/EnhancedInputSubsystems.cpp
@@ -108,7 +108,7 @@ void UEnhancedInputLocalPlayerSubsystem::ControlMappingsRebuiltThisFrame()
 
 bool UEnhancedInputWorldSubsystem::ShouldCreateSubsystem(UObject* Outer) const
 {
-	if (!FSlateApplication::IsInitialized())
+	if (!Super::ShouldCreateSubsystem(Outer) || !FSlateApplication::IsInitialized())
 	{
 		return false;
 	}
@@ -117,10 +117,10 @@ bool UEnhancedInputWorldSubsystem::ShouldCreateSubsystem(UObject* Outer) const
 	const bool bShouldCreate = GetDefault<UEnhancedInputDeveloperSettings>()->bEnableWorldSubsystem;
 	if (!bShouldCreate)
 	{
-		UE_LOG(LogWorldSubsystemInput, Log, TEXT("UEnhancedInputDeveloperSettings::bEnableWorldSubsystem is false, the world subsystem will not be created!"));
+		UE_LOG(LogWorldSubsystemInput, Verbose, TEXT("UEnhancedInputDeveloperSettings::bEnableWorldSubsystem is false, the world subsystem will not be created!"));
 	}
 
-	return bShouldCreate && Super::ShouldCreateSubsystem(Outer);
+	return bShouldCreate;
 }
 
 void UEnhancedInputWorldSubsystem::Initialize(FSubsystemCollectionBase& Collection)
