Fix a TObjectItr that returned too early on the EI world subsystem processor.

#rnx

diff --git a/Engine/Plugins/EnhancedInput/Source/EnhancedInput/Private/EnhancedInputWorldProcessor.cpp b/Engine/Plugins/EnhancedInput/Source/EnhancedInput/Private/EnhancedInputWorldProcessor.cpp
--- a/Engine/Plugins/EnhancedInput/Source/EnhancedInput/Private/EnhancedInputWorldProcessor.cpp
+++ b/Engine/Plugins/EnhancedInput/Source/EnhancedInput/Private/EnhancedInputWorldProcessor.cpp
@@ -180,14 +180,16 @@ void FEnhancedInputWorldProcessor::ProcessAccumulatedPointerInput(float DeltaTim
 
 bool FEnhancedInputWorldProcessor::InputKeyToSubsystem(const FInputKeyParams& Params)
 {
+	bool bRes = false;
+	
 	if (GetDefault<UEnhancedInputDeveloperSettings>()->bEnableWorldSubsystem)
 	{
 		// Tell all the world subsystems about the key that has been pressed
 		for (TObjectIterator<UEnhancedInputWorldSubsystem> It; It; ++It)
 		{
-			return (*It)->InputKey(Params);
+			bRes |= (*It)->InputKey(Params);
 		}
 	}
 	
-	return false;
+	return bRes;
 }
