Update the comments on bIgnoreAllPressedKeysUntilRelease to be clear and provide some examples.

Ensure that we set bShouldBeIgnored AFTER we retain the state of old triggers and modifiers. Otherwise we would never actually update the state of the mapping's bShouldBeIgnored flag and it would always be false. I believe that this issue was introduced with 24566913. Once we got the right behavior for that comparison operator, then the mapping would always be converted instead of the old behavior where it would not. This resulted in it overriding the bShouldBeIgnored flag to whatever the previous state was.

#changelist validated
#virtualized

diff --git a/Engine/Plugins/EnhancedInput/Source/EnhancedInput/Public/EnhancedInputSubsystemInterface.h b/Engine/Plugins/EnhancedInput/Source/EnhancedInput/Public/EnhancedInputSubsystemInterface.h
--- a/Engine/Plugins/EnhancedInput/Source/EnhancedInput/Public/EnhancedInputSubsystemInterface.h
+++ b/Engine/Plugins/EnhancedInput/Source/EnhancedInput/Public/EnhancedInputSubsystemInterface.h
@@ -54,16 +54,47 @@ struct FModifyContextOptions
 		, bNotifyUserSettings(false)
 	{}
 
-	// If true then any keys that are pressed during the rebuild of control mappings will be ignored until they are released.
+	/**
+	 * If true, then any keys that are "down" or "pressed" during the rebuild of control mappings will
+	 * not be processed by Enhanced Input until after they are "released". 
+	 * 
+	 * For example, if you are adding a mapping context with a key mapping to "X",
+	 * and the player is holding down "X" while that IMC is added, 
+	 * there will not be a "Triggered" event until the player releases "X" and presses it again.
+	 * 
+	 * If this is set to false for the above example, then the "Triggered" would fire immediately 
+	 * as soon as the IMC is finished being added.
+	 *
+	 * Default: True
+	 * 
+	 * Note: This will only do something for keys bound to boolean Input Action types.
+	 * Note: This includes all keys that the player has pressed, not just the keys that are previously mapped in Enhanced Input before
+	 * the call to RebuildControlMappings.
+	 */
 	UPROPERTY(EditAnywhere, BlueprintReadWrite, Category = "Input")
 	uint8 bIgnoreAllPressedKeysUntilRelease : 1;
 
-	// The mapping changes will be applied synchronously, rather than at the end of the frame, making them available to the input system on the same frame.
+	/**
+	 * The mapping changes will be applied synchronously, rather than at the end of the frame,
+	 * making them available to the input system on the same frame.
+	 * 
+	 * This is not recommended to be set to true if you are adding multiple mapping contexts 
+	 * as it will have poor performance.
+	 *
+	 * Default: False
+	 */
 	UPROPERTY(EditAnywhere, BlueprintReadWrite, Category = "Input")
 	uint8 bForceImmediately : 1;
 	
-	// If true, then this Mapping Context will be registered or unregistered with the Enhanced Input User Settings on this subsystem,
-	// if they exist.
+	/**
+	 * If true, then this Mapping Context will be registered or unregistered with the
+	 * Enhanced Input User Settings on this subsystem, if they exist.
+	 *
+	 * Default: False
+	 *
+	 * Note: You need to enable and configure your UEnhancedInputUserSettings class in the project
+	 * settings for this to do anything.
+	 */
 	UPROPERTY(EditAnywhere, BlueprintReadWrite, Category = "Input")
 	uint8 bNotifyUserSettings : 1;
 };
