Do not attempt to load the project default IMC assets on the CDO of UEnhancedPlayerInput. Doing so will result in any subobjects of the IMC (triggers, modifiers, player mappable key settings) being null at the time of loadng if they are outside the Enhanced Input module. We can't gaurantee that the load order of these modules such that they would be available before the EI module is being loaded, so we can just skip them here.

#changelist validated
#virtualized

diff --git a/Engine/Plugins/EnhancedInput/Source/EnhancedInput/Private/EnhancedPlayerInput.cpp b/Engine/Plugins/EnhancedInput/Source/EnhancedInput/Private/EnhancedPlayerInput.cpp
--- a/Engine/Plugins/EnhancedInput/Source/EnhancedInput/Private/EnhancedPlayerInput.cpp
+++ b/Engine/Plugins/EnhancedInput/Source/EnhancedInput/Private/EnhancedPlayerInput.cpp
@@ -34,6 +34,13 @@ UEnhancedPlayerInput::UEnhancedPlayerInput()
 	, bIsFlushingInputThisFrame(false)
 	, CurrentlyInUseAnyKeySubstitute(NAME_None)
 {
+	// We don't want to attempt to load stuff on the CDO. Any subobjects on the IMC (triggers, modifiers, Player mappable key settings)
+	// will likely not have been loaded yet if they are defined by the end user or in a different module.
+	if (HasAnyFlags(RF_ClassDefaultObject | RF_ArchetypeObject))
+	{
+		return;
+	}
+
 	if (UE::Input::EnableDefaultMappingContexts)
 	{
 		for (const FDefaultContextSetting& DefaultContext : GetDefault<UEnhancedInputDeveloperSettings>()->DefaultMappingContexts)
