Added missing UPROPERTY on injected input property chain to fix GC-related crash

diff --git a/Engine/Plugins/EnhancedInput/Source/EnhancedInput/Public/EnhancedPlayerInput.h b/Engine/Plugins/EnhancedInput/Source/EnhancedInput/Public/EnhancedPlayerInput.h
--- a/Engine/Plugins/EnhancedInput/Source/EnhancedInput/Public/EnhancedPlayerInput.h
+++ b/Engine/Plugins/EnhancedInput/Source/EnhancedInput/Public/EnhancedPlayerInput.h
@@ -53,6 +53,8 @@ USTRUCT()
 struct FInjectedInputArray
 {
 	GENERATED_BODY()
+
+	UPROPERTY(Transient)
 	TArray<FInjectedInput> Injected;
 };
 
