Only ensure in game builds on the EI User Setting serialize function.

#rnx

diff --git a/Engine/Plugins/EnhancedInput/Source/EnhancedInput/Private/UserSettings/EnhancedInputUserSettings.cpp b/Engine/Plugins/EnhancedInput/Source/EnhancedInput/Private/UserSettings/EnhancedInputUserSettings.cpp
--- a/Engine/Plugins/EnhancedInput/Source/EnhancedInput/Private/UserSettings/EnhancedInputUserSettings.cpp
+++ b/Engine/Plugins/EnhancedInput/Source/EnhancedInput/Private/UserSettings/EnhancedInputUserSettings.cpp
@@ -829,13 +829,27 @@ void UEnhancedInputUserSettings::Serialize(FArchive& Ar)
 		ProfilePair.Value->Serialize(Ar);
 
 		// Save a terminator after each subobject
-		Ar << SavedObjectTerminator;
-		
+		Ar << SavedObjectTerminator;	
+
+		// When you are serializing to an external source, like some other UObject that you want to upload to a custom backend,
+		// this ensure might be hit if that external object has been changed and the data no longer perfectly matches up. 
+		// This can be annoying in the editor, so we will just log instead of ensuring for now to make it less annoying
+		// for devs iterating.
+		// 
+		// For example, if you are working in multiple different release streams or something, and have to bounce backwards
+		// to an older version and your save profile has the newer versions data in it, this ensure would be triggered because
+		// the object terminator does not perfectly match up.
+		// This is something that should never happen in a game build, but can happen quite often in an editor build for devs
+		// with multiple workspaces
+#if WITH_EDITOR
+		UE_CLOG(SavedObjectTerminator != TEXT("ObjectEnd"), LogEnhancedInput, Error, TEXT("Serialization size mismatch! Possible over-read or over-write of this buffer."));
+#else
 		if (!ensure(SavedObjectTerminator == TEXT("ObjectEnd")))
 		{
 			UE_LOG(LogEnhancedInput, Error, TEXT("Serialization size mismatch! Possible over-read or over-write of this buffer."));
 			break;
 		}
+#endif	// WITH_EDITOR
 	}
 }
 
