Do not allow you to make the lower threshold higher then the upper threshold on Dead Zone input modifiers.

diff --git a/Engine/Plugins/EnhancedInput/Source/EnhancedInput/Private/InputModifiers.cpp b/Engine/Plugins/EnhancedInput/Source/EnhancedInput/Private/InputModifiers.cpp
--- a/Engine/Plugins/EnhancedInput/Source/EnhancedInput/Private/InputModifiers.cpp
+++ b/Engine/Plugins/EnhancedInput/Source/EnhancedInput/Private/InputModifiers.cpp
@@ -123,6 +123,37 @@ FLinearColor UInputModifierNegate::GetVisualizationColor_Implementation(FInputAc
 * Dead zones
 */
 
+#if WITH_EDITOR
+EDataValidationResult UInputModifierDeadZone::IsDataValid(class FDataValidationContext& Context) const
+{
+	EDataValidationResult Result = CombineDataValidationResults(Super::IsDataValid(Context), EDataValidationResult::Valid);
+
+	// You cannot scale a boolean value
+	if (LowerThreshold > UpperThreshold)
+	{
+		Result = EDataValidationResult::Invalid;
+		Context.AddError(LOCTEXT("InputModifierDeadZone", "The 'Lower Threshold' cannot be greater then the 'Upper Threshold' of a deadzone."));
+	}
+
+	return Result;
+}
+
+void UInputModifierDeadZone::PostEditChangeProperty(FPropertyChangedEvent& PropertyChangedEvent)
+{
+	Super::PostEditChangeProperty(PropertyChangedEvent);
+
+	const FName MemberPropertyName = PropertyChangedEvent.GetMemberPropertyName();
+	if (MemberPropertyName == GET_MEMBER_NAME_CHECKED(UInputModifierDeadZone, LowerThreshold))
+	{
+		// Clamp the lower threshold to the upper threshold value.
+		if (LowerThreshold > UpperThreshold)
+		{
+			LowerThreshold = UpperThreshold;
+		}
+	}
+}
+#endif	// WITH_EDITOR
+
 FInputActionValue UInputModifierDeadZone::ModifyRaw_Implementation(const UEnhancedPlayerInput* PlayerInput, FInputActionValue CurrentValue, float DeltaTime)
 {
 	EInputActionValueType ValueType = CurrentValue.GetValueType();
diff --git a/Engine/Plugins/EnhancedInput/Source/EnhancedInput/Public/InputModifiers.h b/Engine/Plugins/EnhancedInput/Source/EnhancedInput/Public/InputModifiers.h
--- a/Engine/Plugins/EnhancedInput/Source/EnhancedInput/Public/InputModifiers.h
+++ b/Engine/Plugins/EnhancedInput/Source/EnhancedInput/Public/InputModifiers.h
@@ -132,6 +132,7 @@ class UInputModifierDeadZone : public UInputModifier
 public:
 
 	// Threshold below which input is ignored
+	// This value should always be lower then the UpperThreshold.
 	UPROPERTY(EditInstanceOnly, BlueprintReadWrite, Category=Settings, Config, meta=(ClampMin=0, ClampMax=1))
 	float LowerThreshold = 0.2f;
 
@@ -143,6 +144,12 @@ public:
 	EDeadZoneType Type = EDeadZoneType::Radial;
 
 protected:
+
+#if WITH_EDITOR
+	virtual EDataValidationResult IsDataValid(class FDataValidationContext& Context) const override;
+	virtual void PostEditChangeProperty(FPropertyChangedEvent& PropertyChangedEvent) override;
+#endif
+
 	virtual FInputActionValue ModifyRaw_Implementation(const UEnhancedPlayerInput* PlayerInput, FInputActionValue CurrentValue, float DeltaTime) override;
 
 	// Visualize as black when unmodified. Red when blocked (with differing intensities to indicate axes)
