Remove any input components from the Player Controller's input stack if they are marked as RF_MirroredGarbage. This can happen if you have "Input Enabled" actor added to the world, and it is removed in response to an earlier input component deleting it's owning actor.

diff --git a/Engine/Plugins/EnhancedInput/Source/EnhancedInput/Private/EnhancedPlayerInput.cpp b/Engine/Plugins/EnhancedInput/Source/EnhancedInput/Private/EnhancedPlayerInput.cpp
--- a/Engine/Plugins/EnhancedInput/Source/EnhancedInput/Private/EnhancedPlayerInput.cpp
+++ b/Engine/Plugins/EnhancedInput/Source/EnhancedInput/Private/EnhancedPlayerInput.cpp
@@ -546,7 +546,7 @@ void UEnhancedPlayerInput::EvaluateInputDelegates(const TArray<UInputComponent*>
 		TRACE_CPUPROFILER_EVENT_SCOPE(EnhPIS_Delegates);
 
 		UEnhancedInputComponent* IC = Cast<UEnhancedInputComponent>(InputComponentStack[StackIndex]);
-		if (!IC)
+		if (!IsValid(IC))
 		{
 			continue;
 		}
diff --git a/Engine/Source/Runtime/Engine/Private/PlayerController.cpp b/Engine/Source/Runtime/Engine/Private/PlayerController.cpp
--- a/Engine/Source/Runtime/Engine/Private/PlayerController.cpp
+++ b/Engine/Source/Runtime/Engine/Private/PlayerController.cpp
@@ -2708,7 +2708,7 @@ void APlayerController::BuildInputStack(TArray<UInputComponent*>& InputStack)
 	for (int32 Idx=0; Idx<CurrentInputStack.Num(); ++Idx)
 	{
 		UInputComponent* IC = CurrentInputStack[Idx].Get();
-		if (IC)
+		if (IsValid(IC))
 		{
 			InputStack.Push(IC);
 		}
diff --git a/Engine/Source/Runtime/Engine/Private/UserInterface/PlayerInput.cpp b/Engine/Source/Runtime/Engine/Private/UserInterface/PlayerInput.cpp
--- a/Engine/Source/Runtime/Engine/Private/UserInterface/PlayerInput.cpp
+++ b/Engine/Source/Runtime/Engine/Private/UserInterface/PlayerInput.cpp
@@ -1294,7 +1294,7 @@ void UPlayerInput::EvaluateInputDelegates(const TArray<UInputComponent*>& InputC
 	for ( ; StackIndex >= 0; --StackIndex)
 	{
 		UInputComponent* const IC = InputComponentStack[StackIndex];
-		if (IC)
+		if (IsValid(IC))
 		{
 			check(!KeysToConsume.Num() && !FoundChords.Num() && !EventIndices.Num() && !PotentialActions.Num());
 
