Clamp the lower threshold of deadzone properties if you modify the upper threshold to be higher as well.

#rnx

diff --git a/Engine/Plugins/EnhancedInput/Source/EnhancedInput/Private/InputModifiers.cpp b/Engine/Plugins/EnhancedInput/Source/EnhancedInput/Private/InputModifiers.cpp
--- a/Engine/Plugins/EnhancedInput/Source/EnhancedInput/Private/InputModifiers.cpp
+++ b/Engine/Plugins/EnhancedInput/Source/EnhancedInput/Private/InputModifiers.cpp
@@ -143,7 +143,8 @@ void UInputModifierDeadZone::PostEditChangeProperty(FPropertyChangedEvent& Prope
 	Super::PostEditChangeProperty(PropertyChangedEvent);
 
 	const FName MemberPropertyName = PropertyChangedEvent.GetMemberPropertyName();
-	if (MemberPropertyName == GET_MEMBER_NAME_CHECKED(UInputModifierDeadZone, LowerThreshold))
+	if (MemberPropertyName == GET_MEMBER_NAME_CHECKED(UInputModifierDeadZone, LowerThreshold) ||
+		MemberPropertyName == GET_MEMBER_NAME_CHECKED(UInputModifierDeadZone, UpperThreshold))
 	{
 		// Clamp the lower threshold to the upper threshold value.
 		if (LowerThreshold > UpperThreshold)
