Add a LexToString for ETriggerEvent that will print out the value of a bitmask for the trigger state.

#rnx

diff --git a/Engine/Plugins/EnhancedInput/Source/EnhancedInput/Private/EnhancedInputLibrary.cpp b/Engine/Plugins/EnhancedInput/Source/EnhancedInput/Private/EnhancedInputLibrary.cpp
--- a/Engine/Plugins/EnhancedInput/Source/EnhancedInput/Private/EnhancedInputLibrary.cpp
+++ b/Engine/Plugins/EnhancedInput/Source/EnhancedInput/Private/EnhancedInputLibrary.cpp
@@ -10,6 +10,7 @@
 #include "GameFramework/Actor.h"
 #include "EnhancedInputModule.h"	// For LogEnhancedInput
 #include "EnhancedInputDeveloperSettings.h"
+#include "InputTriggers.h"
 
 #include UE_INLINE_GENERATED_CPP_BY_NAME(EnhancedInputLibrary)
 
@@ -133,3 +134,7 @@ FString UEnhancedInputLibrary::Conv_InputActionValueToString(FInputActionValue A
 	return ActionValue.ToString();
 }
 
+FString UEnhancedInputLibrary::Conv_TriggerEventValueToString(const ETriggerEvent TriggerEvent)
+{
+	return UE::Input::LexToString(TriggerEvent);
+}
\ No newline at end of file
diff --git a/Engine/Plugins/EnhancedInput/Source/EnhancedInput/Private/InputTriggers.cpp b/Engine/Plugins/EnhancedInput/Source/EnhancedInput/Private/InputTriggers.cpp
--- a/Engine/Plugins/EnhancedInput/Source/EnhancedInput/Private/InputTriggers.cpp
+++ b/Engine/Plugins/EnhancedInput/Source/EnhancedInput/Private/InputTriggers.cpp
@@ -12,6 +12,31 @@
 
 #define LOCTEXT_NAMESPACE "EnhancedInputTriggers"
 
+
+namespace UE::Input
+{
+	FString LexToString(const ETriggerEvent TriggerEvent)
+	{
+		if (TriggerEvent == ETriggerEvent::None)
+		{
+			return TEXT("None");
+		}
+
+		FString Result = TEXT("");
+#define TRIGGER_STATE(StatusFlag, DisplayName) if( EnumHasAllFlags(TriggerEvent, StatusFlag) ) Result += (FString(DisplayName) + TEXT("|"));
+		TRIGGER_STATE(ETriggerEvent::Triggered, TEXT("Triggered"));
+		TRIGGER_STATE(ETriggerEvent::Started, TEXT("Started"));
+		TRIGGER_STATE(ETriggerEvent::Ongoing, TEXT("Ongoing"));
+		TRIGGER_STATE(ETriggerEvent::Canceled, TEXT("Canceled"));
+		TRIGGER_STATE(ETriggerEvent::Completed, TEXT("Completed"));
+#undef TRIGGER_STATE
+
+		Result.RemoveFromEnd(TEXT("|"));
+		return Result;
+	}
+}
+
+
 // Abstract trigger bases
 ETriggerState UInputTrigger::UpdateState_Implementation(const UEnhancedPlayerInput* PlayerInput, FInputActionValue ModifiedValue, float DeltaTime)
 {
diff --git a/Engine/Plugins/EnhancedInput/Source/EnhancedInput/Public/EnhancedInputLibrary.h b/Engine/Plugins/EnhancedInput/Source/EnhancedInput/Public/EnhancedInputLibrary.h
--- a/Engine/Plugins/EnhancedInput/Source/EnhancedInput/Public/EnhancedInputLibrary.h
+++ b/Engine/Plugins/EnhancedInput/Source/EnhancedInput/Public/EnhancedInputLibrary.h
@@ -99,6 +99,10 @@ public:
 	UFUNCTION(BlueprintPure, meta = (DisplayName = "To String (InputActionValue)", CompactNodeTitle = "->", BlueprintAutocast), Category = "Utilities|String")
 	static FString Conv_InputActionValueToString(FInputActionValue ActionValue);
 
+	/** Converts an ETriggerEvent to a string */
+	UFUNCTION(BlueprintPure, meta = (DisplayName = "To String (ETriggerEvent)", CompactNodeTitle = "->", BlueprintAutocast), Category = "Utilities|String")
+	static FString Conv_TriggerEventValueToString(const ETriggerEvent TriggerEvent);
+
 	PRAGMA_DISABLE_DEPRECATION_WARNINGS
 
 	UE_DEPRECATED(5.3, "FPlayerMappableKeySlot has been deprecated. Please use EPlayerMappableKeySlot instead.")
diff --git a/Engine/Plugins/EnhancedInput/Source/EnhancedInput/Public/InputTriggers.h b/Engine/Plugins/EnhancedInput/Source/EnhancedInput/Public/InputTriggers.h
--- a/Engine/Plugins/EnhancedInput/Source/EnhancedInput/Public/InputTriggers.h
+++ b/Engine/Plugins/EnhancedInput/Source/EnhancedInput/Public/InputTriggers.h
@@ -54,6 +54,12 @@ enum class ETriggerEvent : uint8
 };
 ENUM_CLASS_FLAGS(ETriggerEvent)
 
+namespace UE::Input
+{
+	// Returns an FString that represents the given ETriggerEvent
+	ENHANCEDINPUT_API FString LexToString(const ETriggerEvent TriggerEvent);
+};
+
 /**
 * Trigger type determine how the trigger contributes to an action's overall trigger event the behavior of the trigger
 */
