Mutate the state of that FInputActionInstance struct to ensure that it is the correct trigger event depending on the delegat when an input action is triggered and started on the same frame.

diff --git a/Engine/Plugins/EnhancedInput/Source/EnhancedInput/Private/EnhancedPlayerInput.cpp b/Engine/Plugins/EnhancedInput/Source/EnhancedInput/Private/EnhancedPlayerInput.cpp
--- a/Engine/Plugins/EnhancedInput/Source/EnhancedInput/Private/EnhancedPlayerInput.cpp
+++ b/Engine/Plugins/EnhancedInput/Source/EnhancedInput/Private/EnhancedPlayerInput.cpp
@@ -625,9 +625,27 @@ void UEnhancedPlayerInput::EvaluateInputDelegates(const TArray<UInputComponent*>
 			if (bCanTrigger)
 			{
 				// Search for the action instance data a second time as a previous delegate call may have deleted it.
-				if (const FInputActionInstance* ActionData = FindActionInstanceData(DelegateAction))
+				if (FInputActionInstance* ActionData = const_cast<FInputActionInstance*>(FindActionInstanceData(DelegateAction)))
 				{
-					Delegate->Execute(*ActionData);
+					// For events that have started and triggered on the same frame, the event will always be 
+					// "Triggered", because that is the latest input state that has been evaluated. While this is the 
+					// correct state, it can be annoying to end users trying to bind the same function to multiple
+					// events and then determine which state they are in, because it will skip the "Started" flag.
+					// By "artificially" setting the trigger event on the action data here we will "force" the 
+					// event to match up to that of the delegate that we are firing.
+					if (ActionData->TriggerEventInternal == ETriggerEventInternal::StartedAndTriggered)
+					{
+						const ETriggerEvent OriginalEvent = ActionData->TriggerEvent;
+						ActionData->TriggerEvent = Delegate->GetTriggerEvent();
+
+						Delegate->Execute(*ActionData);
+
+						ActionData->TriggerEvent = OriginalEvent;
+					}
+					else
+					{
+						Delegate->Execute(*ActionData);
+					}					
 				}
 			}
 		}
