Ensure that the "Canceled" event gets fired if you have a trigger that is ongoing, but gets removed by a mapping context change. Now, instead of removing the input action directly from the ActionInstanceData map, we will reset it's value to zero and mark it for removal. This has the affect of, on the next evaluation of input, we see the value go to zero, and the triggers should evaluate to the canceled state.

This behavior can be disabled by setting the CVar "EnhancedInput.ReconcileRemovedMappingDelegates" to false. This CVar will likely be removed in the future after 2 release versions of the engine to give any dependancies to this broken behavior an opt-out to fix up any broken code.

diff --git a/Engine/Plugins/EnhancedInput/Source/EnhancedInput/Private/EnhancedInputSubsystemInterface.cpp b/Engine/Plugins/EnhancedInput/Source/EnhancedInput/Private/EnhancedInputSubsystemInterface.cpp
--- a/Engine/Plugins/EnhancedInput/Source/EnhancedInput/Private/EnhancedInputSubsystemInterface.cpp
+++ b/Engine/Plugins/EnhancedInput/Source/EnhancedInput/Private/EnhancedInputSubsystemInterface.cpp
@@ -1090,10 +1090,11 @@ void IEnhancedInputSubsystemInterface::RebuildControlMappings()
 				OldMappings.RemoveAtSwap(Idx);
 			}
 		}
-		for (const UInputAction* Action : RemovedActions)
-		{
-			PlayerInput->ActionInstanceData.Remove(Action);
-		}	
+
+		// Actions that are no longer mapped to a key may have been "In progress" by the player
+		// Notify the player input object so that it can reconcile this state and call the "Canceled" event
+		// on the next evaluation of the input.
+		PlayerInput->NotifyInputActionsUnmapped(RemovedActions);
 	}
 
 	// Perform a modifier calculation pass on the default data to initialize values correctly.
diff --git a/Engine/Plugins/EnhancedInput/Source/EnhancedInput/Private/EnhancedPlayerInput.cpp b/Engine/Plugins/EnhancedInput/Source/EnhancedInput/Private/EnhancedPlayerInput.cpp
--- a/Engine/Plugins/EnhancedInput/Source/EnhancedInput/Private/EnhancedPlayerInput.cpp
+++ b/Engine/Plugins/EnhancedInput/Source/EnhancedInput/Private/EnhancedPlayerInput.cpp
@@ -26,6 +26,11 @@ namespace UE
 		static FAutoConsoleVariableRef CVarEnableDefaultMappingContexts(TEXT("EnhancedInput.EnableDefaultMappingContexts"),
 			EnableDefaultMappingContexts,
 			TEXT("Should the UEnhancedInputDeveloperSettings::DefaultMappingContexts be applied to every UEnhancedPlayerInput?"));
+
+		static int32 ReconcileRemovedMappingDelegates = 1;
+		static FAutoConsoleVariableRef CVarReconcileRemovedMappingDelegates(TEXT("EnhancedInput.ReconcileRemovedMappingDelegates"),
+			ReconcileRemovedMappingDelegates,
+			TEXT("When the mappings are rebuilt, they have been 'in process'. If this is true, we will set their value to zero so that they fire the 'Canceled' event on their next evaluation."));
 	}
 }
 
@@ -760,6 +765,16 @@ void UEnhancedPlayerInput::EvaluateInputDelegates(const TArray<UInputComponent*>
 		ActionData.TriggerStateTracker = FTriggerStateTracker();
 	}
 
+	// Remove any input actions that are no longer mapped to the player. At this point they will have been evaluated and 
+	// fired any "Canceled" input events needed
+	for (const UInputAction* IA : ActionsThatHaveBeenRemovedFromMappings)
+	{
+		ActionInstanceData.Remove(IA);
+	}
+
+	// We can clear our queue of mappings that have been removed on the next iteration directly after a key rebuild.
+	ActionsThatHaveBeenRemovedFromMappings.Reset();
+
 	LastFrameTime = CurrentTime;
 	KeysPressedThisTick.Reset();
 	bIsFlushingInputThisFrame = false;
@@ -830,7 +845,8 @@ void UEnhancedPlayerInput::ConditionalBuildKeyMappings_Internal() const
 		if (!LastInjectedActions.Contains(Action) &&
 			!InputsInjectedThisTick.Contains(Action) &&		// This will be empty for most calls, but could potentially contain data.
 			//EngineDefinedActionMappings.ContainsByPredicate(HasActionMapping) && // TODO: EngineDefinedActionMappings are non-rebindable action/key pairings but we have our own systems to handle this...
-			!EnhancedActionMappings.ContainsByPredicate(HasActionMapping))
+			!EnhancedActionMappings.ContainsByPredicate(HasActionMapping) && 
+			!ActionsThatHaveBeenRemovedFromMappings.Contains(Action))
 		{
 			Itr.RemoveCurrent();
 		}
@@ -861,6 +877,37 @@ bool UEnhancedPlayerInput::IsKeyHandledByAction(FKey Key) const
 	return EnhancedKeyBinds.Contains(Key) || Super::IsKeyHandledByAction(Key);
 }
 
+void UEnhancedPlayerInput::NotifyInputActionsUnmapped(const TSet<const UInputAction*>& RemovedInputActions)
+{
+	if (UE::Input::ReconcileRemovedMappingDelegates)
+	{
+		// Instead of totally removing the action instance data so that it stops being processed,
+		// we should instead set the value to zero. This will make it so that upon the next evaluation 
+		// of this action instance data, the "Canceled" event will be fired instead there being
+		// no notification at all. We will queue these action instances for removal from the 
+		// instance data after they have been re-evaluated
+		for (const UInputAction* Action : RemovedInputActions)
+		{
+			if (FInputActionInstance* ActionData = ActionInstanceData.Find(Action))
+			{
+				ActionData->Value.Reset();
+			}
+		}
+
+		ActionsThatHaveBeenRemovedFromMappings = RemovedInputActions;
+	}
+	// ... this is the legacy behavior of just removing the data from the instance data cache
+	else
+	{
+		for (const UInputAction* Action : RemovedInputActions)
+		{			
+			ActionInstanceData.Remove(Action);
+		}
+		
+		UE_LOG(LogEnhancedInput, VeryVerbose, TEXT("[%hs] EnhancedInput.ReconcileRemovedMappingDelegates is false! Using the legacy behavior to remove action instance data. Canceled events will not be fired."), __func__);
+	}
+}
+
 FInputActionInstance& UEnhancedPlayerInput::FindOrAddActionEventData(TObjectPtr<const UInputAction> Action) const
 {
 	FInputActionInstance* Instance = ActionInstanceData.Find(Action);
diff --git a/Engine/Plugins/EnhancedInput/Source/EnhancedInput/Public/EnhancedPlayerInput.h b/Engine/Plugins/EnhancedInput/Source/EnhancedInput/Public/EnhancedPlayerInput.h
--- a/Engine/Plugins/EnhancedInput/Source/EnhancedInput/Public/EnhancedPlayerInput.h
+++ b/Engine/Plugins/EnhancedInput/Source/EnhancedInput/Public/EnhancedPlayerInput.h
@@ -113,6 +113,11 @@ protected:
 	/** This player's version of the Action Mappings */
 	const TArray<FEnhancedActionKeyMapping>& GetEnhancedActionMappings() const { return EnhancedActionMappings; }
 
+	/**
+	* Notifies that the given Input Actions are no longer mapped to any keys, and should have their state reset.
+	*/
+	void NotifyInputActionsUnmapped(const TSet<const UInputAction*>& RemovedInputActions);
+
 	/** Array of data that represents what keys should be consumed if an enhanced input action is in a specific triggered state */
 	UPROPERTY()
 	TMap<TObjectPtr<const UInputAction>, FKeyConsumptionOptions> KeyConsumptionData;
@@ -163,6 +168,9 @@ private:
 	/** Actions that have been triggered this tick and have a delegate that may be fired */
 	TSet<TObjectPtr<const UInputAction>> TriggeredActionsThisTick;
 
+	/** A set of input actions that have been removed from the player's input mappings in a previous rebuild of the key mappings. */
+	TSet<TObjectPtr<const UInputAction>> ActionsThatHaveBeenRemovedFromMappings;
+
 	/**
 	 * A map of Keys to the amount they were depressed this frame. This is reset with each call to ProcessInputStack
 	 * and is populated within UEnhancedPlayerInput::InputKey.
