PR #11801: Small but major fix for Smooth Delta Input Modifier

#rnx

diff --git a/Engine/Plugins/EnhancedInput/Source/EnhancedInput/Private/InputModifiers.cpp b/Engine/Plugins/EnhancedInput/Source/EnhancedInput/Private/InputModifiers.cpp
--- a/Engine/Plugins/EnhancedInput/Source/EnhancedInput/Private/InputModifiers.cpp
+++ b/Engine/Plugins/EnhancedInput/Source/EnhancedInput/Private/InputModifiers.cpp
@@ -20,7 +20,7 @@
 FInputActionValue UInputModifierSmoothDelta::ModifyRaw_Implementation(const UEnhancedPlayerInput* PlayerInput, FInputActionValue CurrentValue, float DeltaTime)
 {
 	// You can't smooth a boolean value
-	if (ensureMsgf(CurrentValue.GetValueType() != EInputActionValueType::Boolean, TEXT("The 'Smooth Delta Modifier' doesn't support boolean values.")))
+	if (!ensureMsgf(CurrentValue.GetValueType() != EInputActionValueType::Boolean, TEXT("The 'Smooth Delta Modifier' doesn't support boolean values.")))
 	{
 		return CurrentValue;
 	}
