Add a "FlushPressedKeys" function to the Enhanced Input BP library

#rnx

diff --git a/Engine/Plugins/EnhancedInput/Source/EnhancedInput/Private/EnhancedInputLibrary.cpp b/Engine/Plugins/EnhancedInput/Source/EnhancedInput/Private/EnhancedInputLibrary.cpp
--- a/Engine/Plugins/EnhancedInput/Source/EnhancedInput/Private/EnhancedInputLibrary.cpp
+++ b/Engine/Plugins/EnhancedInput/Source/EnhancedInput/Private/EnhancedInputLibrary.cpp
@@ -11,6 +11,7 @@
 #include "EnhancedInputModule.h"	// For LogEnhancedInput
 #include "EnhancedInputDeveloperSettings.h"
 #include "InputTriggers.h"
+#include "GameFramework/PlayerController.h"
 
 #include UE_INLINE_GENERATED_CPP_BY_NAME(EnhancedInputLibrary)
 
@@ -137,4 +138,12 @@ FString UEnhancedInputLibrary::Conv_InputActionValueToString(FInputActionValue A
 FString UEnhancedInputLibrary::Conv_TriggerEventValueToString(const ETriggerEvent TriggerEvent)
 {
 	return UE::Input::LexToString(TriggerEvent);
+}
+
+void UEnhancedInputLibrary::FlushPlayerInput(APlayerController* PlayerController)
+{
+	if (PlayerController)
+	{
+		PlayerController->FlushPressedKeys();
+	}
 }
\ No newline at end of file
diff --git a/Engine/Plugins/EnhancedInput/Source/EnhancedInput/Public/EnhancedInputLibrary.h b/Engine/Plugins/EnhancedInput/Source/EnhancedInput/Public/EnhancedInputLibrary.h
--- a/Engine/Plugins/EnhancedInput/Source/EnhancedInput/Public/EnhancedInputLibrary.h
+++ b/Engine/Plugins/EnhancedInput/Source/EnhancedInput/Public/EnhancedInputLibrary.h
@@ -122,6 +122,14 @@ public:
 	static FPlayerMappableKeySlot& GetFourthPlayerMappableKeySlot() { return FPlayerMappableKeySlot::FourthKeySlot; };
 	
 	PRAGMA_ENABLE_DEPRECATION_WARNINGS
+
+	/**
+	 * Flushes the player controller's pressed keys
+	 * 
+	 * @see APlayerController::FlushPressedKeys
+	 */
+	UFUNCTION(BlueprintCallable, Category="Input")
+	static void FlushPlayerInput(APlayerController* PlayerController);
 };
 
 #if UE_ENABLE_INCLUDE_ORDER_DEPRECATED_IN_5_2
