Enhanced Input optimization: Cache if a key mapping has any "Always tick" triggers set on it instead of looping through all triggers for each key mapping in UEnhancedPlayerInput::ProcessActionMappingEvent

#rnx

diff --git a/Engine/Plugins/EnhancedInput/Source/EnhancedInput/Private/EnhancedActionKeyMapping.cpp b/Engine/Plugins/EnhancedInput/Source/EnhancedInput/Private/EnhancedActionKeyMapping.cpp
--- a/Engine/Plugins/EnhancedInput/Source/EnhancedInput/Private/EnhancedActionKeyMapping.cpp
+++ b/Engine/Plugins/EnhancedInput/Source/EnhancedInput/Private/EnhancedActionKeyMapping.cpp
@@ -47,6 +47,7 @@ FEnhancedActionKeyMapping::FEnhancedActionKeyMapping(const UInputAction* InActio
 	, Action(InAction)
 	, Key(InKey)
 	, bShouldBeIgnored(false)
+	, bHasAlwaysTickTrigger(false)
 	, bIsPlayerMappable(false)
 {}
 PRAGMA_ENABLE_DEPRECATION_WARNINGS
@@ -57,6 +58,7 @@ FEnhancedActionKeyMapping::FEnhancedActionKeyMapping(const UInputAction* InActio
 	: Action(InAction)
 	, Key(InKey)
 	, bShouldBeIgnored(false)
+	, bHasAlwaysTickTrigger(false)
 {}
 
 #endif	// WITH_EDITORONLY_DATA
diff --git a/Engine/Plugins/EnhancedInput/Source/EnhancedInput/Private/EnhancedInputSubsystemInterface.cpp b/Engine/Plugins/EnhancedInput/Source/EnhancedInput/Private/EnhancedInputSubsystemInterface.cpp
--- a/Engine/Plugins/EnhancedInput/Source/EnhancedInput/Private/EnhancedInputSubsystemInterface.cpp
+++ b/Engine/Plugins/EnhancedInput/Source/EnhancedInput/Private/EnhancedInputSubsystemInterface.cpp
@@ -776,11 +776,17 @@ TArray<FEnhancedActionKeyMapping> IEnhancedInputSubsystemInterface::ReorderMappi
 {
 	TSet<const UInputAction*> ChordingActions;
 
-	// Gather all chording actions within a mapping's triggers.
-	auto GatherChordingActions = [&ChordingActions, &DependentChordActions](const FEnhancedActionKeyMapping& Mapping)
+	struct FTriggerEvaluationResults
 	{
 		bool bFoundChordTrigger = false;
-		auto EvaluateTriggers = [&Mapping, &ChordingActions, &bFoundChordTrigger, &DependentChordActions](const TArray<UInputTrigger*>& Triggers)
+		bool bFoundAlwaysTickTrigger = false;
+	};
+	
+	// Gather all chording actions within a mapping's triggers.
+	auto GatherChordingActions = [&ChordingActions, &DependentChordActions](const FEnhancedActionKeyMapping& Mapping) -> FTriggerEvaluationResults
+	{
+		FTriggerEvaluationResults Res = {};
+		auto EvaluateTriggers = [&Mapping, &ChordingActions, &DependentChordActions, &Res](const TArray<UInputTrigger*>& Triggers)-> FTriggerEvaluationResults
 		{
 			for (const UInputTrigger* Trigger : Triggers)
 			{
@@ -791,18 +797,27 @@ TArray<FEnhancedActionKeyMapping> IEnhancedInputSubsystemInterface::ReorderMappi
 					// Keep track of the action itself, and the action it is dependant on
 					DependentChordActions.Emplace(UEnhancedPlayerInput::FDependentChordTracker { Mapping.Action, ChordTrigger->ChordAction });
 					
-					bFoundChordTrigger = true;
+					Res.bFoundChordTrigger = true;
 				}
+
+				// Keep track of if this trigger is marked as being "always tick".
+				// This is not a great thing to do but some custom triggers may require always being ticked, so allow it as an option
+				Res.bFoundAlwaysTickTrigger |= Trigger->bShouldAlwaysTick;
 			}
+			return Res;
 		};
-		EvaluateTriggers(Mapping.Triggers);
 		
-		if(ensureMsgf(Mapping.Action, TEXT("A key mapping has no associated action!")))
-		{
-			EvaluateTriggers(Mapping.Action->Triggers);			
-		}
+		const FTriggerEvaluationResults MappingResults = EvaluateTriggers(Mapping.Triggers);
+
+		ensureMsgf(Mapping.Action, TEXT("A key mapping has no associated action!"));
+		const FTriggerEvaluationResults ActionResults = EvaluateTriggers(Mapping.Action->Triggers);
 
-		return bFoundChordTrigger;
+		// returned the combined results of each individual keymapping and it's associated input action.
+		return FTriggerEvaluationResults
+		{
+			.bFoundChordTrigger			= (MappingResults.bFoundChordTrigger || ActionResults.bFoundChordTrigger),
+			.bFoundAlwaysTickTrigger	= (MappingResults.bFoundAlwaysTickTrigger || ActionResults.bFoundAlwaysTickTrigger)
+		};
 	};
 
 	// Split chorded mappings (second priority) from all others whilst building a list of chording actions to use for further prioritization.
@@ -812,10 +827,17 @@ TArray<FEnhancedActionKeyMapping> IEnhancedInputSubsystemInterface::ReorderMappi
 	int32 NumEmptyMappings = 0;
 	for (const FEnhancedActionKeyMapping& Mapping : UnorderedMappings)
 	{
-		if(Mapping.Action)
+		if (Mapping.Action)
 		{
-			TArray<FEnhancedActionKeyMapping>& MappingArray = GatherChordingActions(Mapping) ? ChordedMappings : OtherMappings;
-			MappingArray.Add(Mapping);
+			// Evaluate the triggers on each key mapping to check for chords and also "always tick" input triggers.
+			const FTriggerEvaluationResults TriggerEvalResults = GatherChordingActions(Mapping);
+
+			// Determine which array this mapping should be in based on if it has a chord or not
+			TArray<FEnhancedActionKeyMapping>& MappingArray = TriggerEvalResults.bFoundChordTrigger ? ChordedMappings : OtherMappings;
+
+			// flag this new mapping as being always tick as necessary
+			FEnhancedActionKeyMapping& NewlyAddedMapping = MappingArray.Add_GetRef(Mapping);
+			NewlyAddedMapping.bHasAlwaysTickTrigger = TriggerEvalResults.bFoundAlwaysTickTrigger;
 		}
 		else
 		{
@@ -829,7 +851,8 @@ TArray<FEnhancedActionKeyMapping> IEnhancedInputSubsystemInterface::ReorderMappi
 
 	// Move chording mappings to the front as they need to be evaluated before chord and blocker triggers
 	// TODO: Further ordering of chording mappings may be required should one of them be chorded against another
-	auto ExtractChords = [&OrderedMappings, &ChordingActions](TArray<FEnhancedActionKeyMapping>& Mappings) {
+	auto ExtractChords = [&OrderedMappings, &ChordingActions](TArray<FEnhancedActionKeyMapping>& Mappings)
+	{
 		for (int32 i = 0; i < Mappings.Num();)
 		{
 			if (ChordingActions.Contains(Mappings[i].Action))
diff --git a/Engine/Plugins/EnhancedInput/Source/EnhancedInput/Private/EnhancedPlayerInput.cpp b/Engine/Plugins/EnhancedInput/Source/EnhancedInput/Private/EnhancedPlayerInput.cpp
--- a/Engine/Plugins/EnhancedInput/Source/EnhancedInput/Private/EnhancedPlayerInput.cpp
+++ b/Engine/Plugins/EnhancedInput/Source/EnhancedInput/Private/EnhancedPlayerInput.cpp
@@ -159,8 +159,16 @@ enum class EKeyEvent : uint8
 	Held,		// Key generated no event, but is in a held state and wants to continue applying modifiers and triggers
 };
 
-void UEnhancedPlayerInput::ProcessActionMappingEvent(TObjectPtr<const UInputAction> Action, float DeltaTime, bool bGamePaused, FInputActionValue RawKeyValue, EKeyEvent KeyEvent, const TArray<UInputModifier*>& Modifiers, const TArray<UInputTrigger*>& Triggers)
-{
+void UEnhancedPlayerInput::ProcessActionMappingEvent(
+	TObjectPtr<const UInputAction> Action,
+	float DeltaTime,
+	bool bGamePaused,
+	FInputActionValue RawKeyValue,
+	EKeyEvent KeyEvent,
+	const TArray<UInputModifier*>& Modifiers,
+	const TArray<UInputTrigger*>& Triggers,
+	const bool bHasAlwaysTickTrigger /*= false*/)
+{	
 	FInputActionInstance& ActionData = FindOrAddActionEventData(Action);
 
 	// Update values and triggers for all actionable mappings each frame
@@ -169,32 +177,9 @@ void UEnhancedPlayerInput::ProcessActionMappingEvent(TObjectPtr<const UInputActi
 	// Reset action data on the first event processed for the action this tick.
 	bool bResetActionData = !ActionsWithEventsThisTick.Contains(Action);
 	bool bMappingTriggersApplied = false;
-
-	bool bHasAnyAlwaysTickTriggers = false;
-	// checking the input mapping context triggers for any triggers that should tick every frame
-	for (const UInputTrigger* Trigger : Triggers)
-	{
-		if (Trigger && Trigger->bShouldAlwaysTick)
-		{
-			bHasAnyAlwaysTickTriggers = true;
-			break;
-		}
-	}
-	// we also need to check the triggers of the Input Action itself - only if we haven't already found an AlwaysTickTrigger
-	if (!bHasAnyAlwaysTickTriggers)
-	{
-		for (const UInputTrigger* Trigger : Action->Triggers)
-		{
-			if (Trigger && Trigger->bShouldAlwaysTick)
-			{
-				bHasAnyAlwaysTickTriggers = true;
-				break;
-			}
-		}
-	}
 	
 	// If the key state is changing or the key is actuated and being held (and not coming back up this tick) recalculate its value and resulting trigger state.
-	if (KeyEvent != EKeyEvent::None || bHasAnyAlwaysTickTriggers)
+	if (KeyEvent != EKeyEvent::None || bHasAlwaysTickTrigger)
 	{
 		if (bResetActionData)
 		{
@@ -442,7 +427,7 @@ void UEnhancedPlayerInput::EvaluateInputDelegates(const TArray<UInputComponent*>
 		}
 
 		// Perform update
-		ProcessActionMappingEvent(Mapping.Action, NonDilatedDeltaTime, bGamePaused, RawKeyValue, KeyEvent, Mapping.Modifiers, Mapping.Triggers);
+		ProcessActionMappingEvent(Mapping.Action, NonDilatedDeltaTime, bGamePaused, RawKeyValue, KeyEvent, Mapping.Modifiers, Mapping.Triggers, Mapping.bHasAlwaysTickTrigger);
 	}
 
 
@@ -459,7 +444,16 @@ void UEnhancedPlayerInput::EvaluateInputDelegates(const TArray<UInputComponent*>
 		else if (!InputsInjectedThisTick.Contains(InjectedAction))
 		{
 			// Reset action state by "releasing the key".
-			ProcessActionMappingEvent(InjectedAction, NonDilatedDeltaTime, bGamePaused, FInputActionValue(), EKeyEvent::Actuated, {}, {});
+			ProcessActionMappingEvent(
+				InjectedAction,
+				NonDilatedDeltaTime,
+				bGamePaused,
+				FInputActionValue(),
+				EKeyEvent::Actuated,
+				{},
+				{},
+				/* bHasAlwaysTickTrigger= */ false);
+			
 			It.RemoveCurrent();
 		}
 	}
diff --git a/Engine/Plugins/EnhancedInput/Source/EnhancedInput/Public/EnhancedActionKeyMapping.h b/Engine/Plugins/EnhancedInput/Source/EnhancedInput/Public/EnhancedActionKeyMapping.h
--- a/Engine/Plugins/EnhancedInput/Source/EnhancedInput/Public/EnhancedActionKeyMapping.h
+++ b/Engine/Plugins/EnhancedInput/Source/EnhancedInput/Public/EnhancedActionKeyMapping.h
@@ -160,6 +160,17 @@ struct ENHANCEDINPUT_API FEnhancedActionKeyMapping
 	UPROPERTY(Transient)
 	uint8 bShouldBeIgnored : 1;
 
+	/**
+	 * True if any triggers on this mapping or its associated UInputAction
+	 * are flagged as "always tick" triggers. This is only set when control
+	 * mappings are rebuilt.
+	 *
+	 * @see UInputTrigger::bShouldAlwaysTick
+	 * @see IEnhancedInputSubsystemInterface::ReorderMappings
+	 */
+	UPROPERTY(Transient)
+	uint8 bHasAlwaysTickTrigger : 1;
+
 #if WITH_EDITORONLY_DATA
 
 	/** If true then this ActionKeyMapping will be exposed as a player mappable key */
diff --git a/Engine/Plugins/EnhancedInput/Source/EnhancedInput/Public/EnhancedPlayerInput.h b/Engine/Plugins/EnhancedInput/Source/EnhancedInput/Public/EnhancedPlayerInput.h
--- a/Engine/Plugins/EnhancedInput/Source/EnhancedInput/Public/EnhancedPlayerInput.h
+++ b/Engine/Plugins/EnhancedInput/Source/EnhancedInput/Public/EnhancedPlayerInput.h
@@ -138,7 +138,16 @@ private:
 	FInputActionValue ApplyModifiers(const TArray<UInputModifier*>& Modifiers, FInputActionValue RawValue, float DeltaTime) const;						// Pre-modified (raw) value
 	ETriggerEventInternal GetTriggerStateChangeEvent(ETriggerState LastTriggerState, ETriggerState NewTriggerState) const;
 	ETriggerEvent ConvertInternalTriggerEvent(ETriggerEventInternal Event) const;	// Collapse a detailed internal trigger event into a friendly representation
-	void ProcessActionMappingEvent(TObjectPtr<const UInputAction> Action, float DeltaTime, bool bGamePaused, FInputActionValue RawValue, EKeyEvent KeyEvent, const TArray<UInputModifier*>& Modifiers, const TArray<UInputTrigger*>& Triggers);
+
+	void ProcessActionMappingEvent(
+		TObjectPtr<const UInputAction> Action,
+		float DeltaTime,
+		bool bGamePaused,
+		FInputActionValue RawValue,
+		EKeyEvent KeyEvent,
+		const TArray<UInputModifier*>& Modifiers,
+		const TArray<UInputTrigger*>& Triggers,
+		const bool bHasAlwaysTickTrigger = false);
 
 	FInputActionInstance& FindOrAddActionEventData(TObjectPtr<const UInputAction> Action) const;
 
