Add a test for the performance of evaluating input modifiers in Enhanced Input. This test will just tick the player input with a few input modifiers a bunch of times and report the min, max, and median timings for it.

This will make for a good A/B test when optimizing modifiers, which at the moment are the most expensive part of enhanced input ticking.

#rnx

diff --git a/Engine/Plugins/EnhancedInput/Source/InputEditor/Private/Tests/InputModifierTest.cpp b/Engine/Plugins/EnhancedInput/Source/InputEditor/Private/Tests/InputModifierTest.cpp
--- a/Engine/Plugins/EnhancedInput/Source/InputEditor/Private/Tests/InputModifierTest.cpp
+++ b/Engine/Plugins/EnhancedInput/Source/InputEditor/Private/Tests/InputModifierTest.cpp
@@ -1,8 +1,10 @@
 // Copyright Epic Games, Inc. All Rights Reserved.
 
 #include "InputTestFramework.h"
+#include "Algo/MaxElement.h"
 #include "Misc/AutomationTest.h"
-
+#include "EnhancedInputModule.h"
+#include "ProfilingDebugging/ScopedTimers.h"
 // Tests focused on individual modifiers
 
 
@@ -201,3 +203,73 @@ bool FInputModifierDeadzoneTest::RunTest(const FString& Parameters)
 
 	return true;
 }
+
+/**
+ * A simple "stress test" for Enhanced Input modifiers that we can use to measure the performance
+ * of applying multiple UInputModifiers to a single key mapping.
+ */
+IMPLEMENT_SIMPLE_AUTOMATION_TEST(FInputModifierPerformanceTest, "Input.Modifiers.Performance", BasicModifierTestFlags)
+
+bool FInputModifierPerformanceTest::RunTest(const FString& Parameters)
+{
+	GIVEN(UControllablePlayer& Data = ABasicModifierTest(this, EInputActionValueType::Axis3D));
+	AND(AnActionIsMappedToAKey(Data, TestContext, TestAction, TestAxis));
+
+	// Apply a few different types of input modifiers to this key
+	AND(AModifierIsAppliedToAnActionMapping(Data, NewObject<UInputModifierScalar>(), TestContext, TestAction, TestAxis));
+	AND(AModifierIsAppliedToAnActionMapping(Data, NewObject<UInputModifierNegate>(), TestContext, TestAction, TestAxis));
+	AND(AModifierIsAppliedToAnActionMapping(Data, NewObject<UInputModifierDeadZone>(), TestContext, TestAction, TestAxis));
+	AND(AModifierIsAppliedToAnActionMapping(Data, NewObject<UInputModifierSwizzleAxis>(), TestContext, TestAction, TestAxis));
+	AND(AModifierIsAppliedToAnActionMapping(Data, NewObject<UInputModifierFOVScaling>(), TestContext, TestAction, TestAxis));
+
+	// Test applying a a key value of some kind
+	WHEN(AKeyIsActuated(Data, TestAxis, 0.84648f));
+
+	auto RunPerfTick = [&Data]() -> double
+	{
+		double Duration = 0.0;
+		{
+			FDurationTimer Timer{ Duration };
+			Timer.Start();
+	
+			static constexpr int32 NumTicksToMeasure = 10000;
+			for (int32 i = 0; i < NumTicksToMeasure; ++i)
+			{
+				// Tick the input stack, which will call the "ModifyRaw" function on every modifier we have
+				AND(InputIsTicked(Data));		
+			}
+		
+			Timer.Stop();
+		}
+		return Duration;
+	};
+	
+	// Tick all these modifiers a bunch of times
+	static constexpr int32 NumTimesToRun = 30;
+
+	TArray<double, TInlineAllocator<NumTimesToRun>> AverageTimings = {};
+	for (int32 i = 0; i < NumTimesToRun; ++i)
+	{
+		AverageTimings.Emplace(RunPerfTick());		
+	}
+
+	AverageTimings.Sort([](const double A, const double B)
+	{
+		return A < B;
+	});
+
+	const double MinRun = *Algo::MinElement(AverageTimings);
+	const double MaxRun = *Algo::MaxElement(AverageTimings);
+	
+	// Note; this requires that the AverageTimings array has an even number of elements in it...
+	// Otherwise you could just do AverageTimings[NumTimesToRun / 2] and get the middle element
+	check(AverageTimings.Num() % 2 == 0);
+	const double MedianTiming = (AverageTimings[NumTimesToRun / 2 - 1] + AverageTimings[NumTimesToRun / 2]) / 2;;
+
+	UE_LOG(LogEnhancedInput, Log, TEXT("Modifiers Perf Test (in seconds)... Median: %lf   Min: %lf   Max: %lf"), MedianTiming, MinRun, MaxRun);
+	
+	TestLessEqual(TEXT("MinRun is less then max"), MinRun, MaxRun);
+	TestNotEqual(TEXT("Median Time was calculated"), MedianTiming, 0.0);
+	
+	return true;
+}
\ No newline at end of file
