Enhanced Input: Fixing issue with combo triggers not firing completed event when first and last combo action are the same action

diff --git a/Engine/Plugins/EnhancedInput/Source/EnhancedInput/Private/InputTriggers.cpp b/Engine/Plugins/EnhancedInput/Source/EnhancedInput/Private/InputTriggers.cpp
--- a/Engine/Plugins/EnhancedInput/Source/EnhancedInput/Private/InputTriggers.cpp
+++ b/Engine/Plugins/EnhancedInput/Source/EnhancedInput/Private/InputTriggers.cpp
@@ -3,6 +3,7 @@
 #include "InputTriggers.h"
 #include "EnhancedInputModule.h"
 #include "EnhancedPlayerInput.h"
+#include "HAL/IConsoleManager.h"
 
 #if WITH_EDITOR
 #include "Misc/DataValidation.h"
@@ -12,6 +13,17 @@
 
 #define LOCTEXT_NAMESPACE "EnhancedInputTriggers"
 
+namespace UE::EnhancedInput
+{
+	static TAutoConsoleVariable<bool> CVarCheckInitialStateForComboTrigger
+	(
+		TEXT("EnhancedInput.Triggers.bCheckInitalStateForComboTrigger"),
+		false,
+		TEXT("When true, combo triggers will check initial state (First input action in the combo array) and return 'Ongoing' if the Combo Step Completion State is met. ")
+		TEXT("Note: Setting this to true was Combo trigger behavior as of 5.4 and before."),
+		ECVF_Default
+	);
+}
 
 namespace UE::Input
 {
@@ -331,13 +343,17 @@ ETriggerState UInputTriggerCombo::UpdateState_Implementation(const UEnhancedPlay
 		{
 			return ETriggerState::Ongoing;
 		}
-	
-		// Really should account for first combo action being mid-trigger...
-		const FInputActionInstance* InitialState = PlayerInput->FindActionInstanceData(ComboActions[0].ComboStepAction);
-		if (InitialState && InitialState->GetTriggerEvent() > ETriggerEvent::None) // || Cancelled!
+
+		if (UE::EnhancedInput::CVarCheckInitialStateForComboTrigger.GetValueOnAnyThread())
 		{
-			return ETriggerState::Ongoing;
+			// Really should account for first combo action being mid-trigger...
+			const FInputActionInstance* InitialState = PlayerInput->FindActionInstanceData(ComboActions[0].ComboStepAction);
+			if (InitialState && InitialState->GetTriggerEvent() > ETriggerEvent::None) // || Cancelled!
+			{
+				return ETriggerState::Ongoing;
+			}
 		}
+		
 		CurrentTimeBetweenComboSteps = 0;
 	}
 	return ETriggerState::None;
