Handle the case where you change the Enhanced input User settings class in the editor and now need to create a different settings save object.

diff --git a/Engine/Plugins/EnhancedInput/Source/EnhancedInput/Private/UserSettings/EnhancedInputUserSettings.cpp b/Engine/Plugins/EnhancedInput/Source/EnhancedInput/Private/UserSettings/EnhancedInputUserSettings.cpp
--- a/Engine/Plugins/EnhancedInput/Source/EnhancedInput/Private/UserSettings/EnhancedInputUserSettings.cpp
+++ b/Engine/Plugins/EnhancedInput/Source/EnhancedInput/Private/UserSettings/EnhancedInputUserSettings.cpp
@@ -586,11 +586,25 @@ UEnhancedInputUserSettings* UEnhancedInputUserSettings::LoadOrCreateSettings(ULo
 
 	// If there is no settings save game object, then we can create on
 	// based on the class type set in the developer settings
-	if (Settings == nullptr)
+	const UEnhancedInputDeveloperSettings* DevSettings = GetDefault<UEnhancedInputDeveloperSettings>();
+	UClass* SettingsClass = DevSettings->UserSettingsClass ? DevSettings->UserSettingsClass.Get() : UEnhancedInputDeveloperSettings::StaticClass();
+
+	// This property is marked as "NoClear", so this should be impossible.
+	if (!ensureMsgf(SettingsClass, TEXT("Invalid Enhanced Input User settings class!")))
+	{
+		return nullptr;
+	}
+	
+	UE_CLOG((Settings && (Settings->GetClass() != SettingsClass)),
+		LogEnhancedInput,
+		Log,
+		TEXT("[%hs] Enhanced Input User Settings class has been changed from '%s' to '%s'. A new settings save object will be created (your saved settings will be reset)"),
+		__func__, *GetNameSafe(Settings->GetClass()), *GetNameSafe(SettingsClass));
+	
+	// If the settings are null (they dont exist yet) or the class has changed, we need to create a new object.
+	// The class can change if you modify it in the editor to be something else
+	if (Settings == nullptr || (Settings->GetClass() != SettingsClass))
 	{
-		const UEnhancedInputDeveloperSettings* DevSettings = GetDefault<UEnhancedInputDeveloperSettings>();
-		UClass* SettingsClass = DevSettings->UserSettingsClass ? DevSettings->UserSettingsClass.Get() : UEnhancedInputDeveloperSettings::StaticClass();
-		
 		Settings = Cast<UEnhancedInputUserSettings>(UGameplayStatics::CreateSaveGameObject(SettingsClass));
 	}
 
