Fix the EI User settings using the wrong default class. It was referring to the developer settings instead of the user settings!

#rnx

diff --git a/Engine/Plugins/EnhancedInput/Source/EnhancedInput/Private/UserSettings/EnhancedInputUserSettings.cpp b/Engine/Plugins/EnhancedInput/Source/EnhancedInput/Private/UserSettings/EnhancedInputUserSettings.cpp
--- a/Engine/Plugins/EnhancedInput/Source/EnhancedInput/Private/UserSettings/EnhancedInputUserSettings.cpp
+++ b/Engine/Plugins/EnhancedInput/Source/EnhancedInput/Private/UserSettings/EnhancedInputUserSettings.cpp
@@ -587,7 +587,7 @@ UEnhancedInputUserSettings* UEnhancedInputUserSettings::LoadOrCreateSettings(ULo
 	// If there is no settings save game object, then we can create on
 	// based on the class type set in the developer settings
 	const UEnhancedInputDeveloperSettings* DevSettings = GetDefault<UEnhancedInputDeveloperSettings>();
-	UClass* SettingsClass = DevSettings->UserSettingsClass ? DevSettings->UserSettingsClass.Get() : UEnhancedInputDeveloperSettings::StaticClass();
+	UClass* SettingsClass = DevSettings->UserSettingsClass ? DevSettings->UserSettingsClass.Get() : UEnhancedInputUserSettings::StaticClass();
 
 	// This property is marked as "NoClear", so this should be impossible.
 	if (!ensureMsgf(SettingsClass, TEXT("Invalid Enhanced Input User settings class!")))
