[Input] Updating scale input modifier to be EditAnywhere.
-This allows non instanced objects (such as widgets that will inject touch input), to be able to use the scale modifier.

#tests Changed to EditAnywhere, and observed that no data loss occured in PIE

diff --git a/Engine/Plugins/EnhancedInput/Source/EnhancedInput/Public/InputModifiers.h b/Engine/Plugins/EnhancedInput/Source/EnhancedInput/Public/InputModifiers.h
--- a/Engine/Plugins/EnhancedInput/Source/EnhancedInput/Public/InputModifiers.h
+++ b/Engine/Plugins/EnhancedInput/Source/EnhancedInput/Public/InputModifiers.h
@@ -178,7 +178,7 @@ public:
 	 *  
 	 * Note: This will do nothing on boolean input action types, as they can only be true or false.
 	 */
-	UPROPERTY(EditInstanceOnly, BlueprintReadWrite, Category=Settings)
+	UPROPERTY(EditAnywhere, BlueprintReadWrite, Category=Settings)
 	FVector Scalar = FVector::OneVector;
 
 protected:
