Fix a crash when having a null trigger on a chorded input action.

#rnx

diff --git a/Engine/Plugins/EnhancedInput/Source/EnhancedInput/Private/EnhancedInputSubsystemInterface.cpp b/Engine/Plugins/EnhancedInput/Source/EnhancedInput/Private/EnhancedInputSubsystemInterface.cpp
--- a/Engine/Plugins/EnhancedInput/Source/EnhancedInput/Private/EnhancedInputSubsystemInterface.cpp
+++ b/Engine/Plugins/EnhancedInput/Source/EnhancedInput/Private/EnhancedInputSubsystemInterface.cpp
@@ -790,6 +790,12 @@ TArray<FEnhancedActionKeyMapping> IEnhancedInputSubsystemInterface::ReorderMappi
 		{
 			for (const UInputTrigger* Trigger : Triggers)
 			{
+				if (!Trigger)
+				{
+					UE_LOG(LogEnhancedInput, Error, TEXT("Null input trigger detected in mapping to input action '%s'"), *GetNameSafe(Mapping.Action));
+					continue;
+				}
+				
 				if (const UInputTriggerChordAction* ChordTrigger = Cast<const UInputTriggerChordAction>(Trigger))
 				{
 					ChordingActions.Add(ChordTrigger->ChordAction);
