Remove the additional "console variable" metadata from the EI editor settings as they did not really work intuitively. This will ensure that your editor prefs persist between sessions. If you wish to set these settings in a config file, then you can do so with the normal object sytax:

  [/Script/InputEditor.EnhancedInputEditorProjectSettings]
  bAutomaticallyStartConsumingInput=true
  bLogAllInput=true

diff --git a/Engine/Plugins/EnhancedInput/Source/InputEditor/Private/EnhancedInputEditorSettings.cpp b/Engine/Plugins/EnhancedInput/Source/InputEditor/Private/EnhancedInputEditorSettings.cpp
--- a/Engine/Plugins/EnhancedInput/Source/InputEditor/Private/EnhancedInputEditorSettings.cpp
+++ b/Engine/Plugins/EnhancedInput/Source/InputEditor/Private/EnhancedInputEditorSettings.cpp
@@ -3,30 +3,9 @@
 #include "EnhancedInputEditorSettings.h"
 #include "EnhancedInputDeveloperSettings.h"
 #include "EnhancedPlayerInput.h"
-#include "HAL/IConsoleManager.h"
 
 #include UE_INLINE_GENERATED_CPP_BY_NAME(EnhancedInputEditorSettings)
 
-namespace UE::EnhancedInput::Private
-{
-	namespace ConsoleVariables
-	{
-		static bool bShouldLogAllInputs = false;
-		static FAutoConsoleVariableRef CVarShouldLogAllInputs(
-			TEXT("EnhancedEditorInput.bShouldLogAllInputs"),
-			bShouldLogAllInputs,
-			TEXT("Should each InputKey call be logged?"),
-			ECVF_Default);
-
-		static bool bAutomaticallyStartConsumingInput = false;
-		static FAutoConsoleVariableRef CVarAutomaticallyStartConsumingInput(
-			TEXT("EnhancedEditorInput.bAutomaticallyStartConsumingInput"),
-			bAutomaticallyStartConsumingInput,
-			TEXT("Should the UEnhancedInputEditorSubsystem be started as soon as it is inialized?"),
-			ECVF_Default);	
-	}
-}
-
 ////////////////////////////////////////////////////////////////////////
 // UEnhancedInputEditorProjectSettings
 UEnhancedInputEditorProjectSettings::UEnhancedInputEditorProjectSettings(const FObjectInitializer& Initializer)
diff --git a/Engine/Plugins/EnhancedInput/Source/InputEditor/Public/EnhancedInputEditorSettings.h b/Engine/Plugins/EnhancedInput/Source/InputEditor/Public/EnhancedInputEditorSettings.h
--- a/Engine/Plugins/EnhancedInput/Source/InputEditor/Public/EnhancedInputEditorSettings.h
+++ b/Engine/Plugins/EnhancedInput/Source/InputEditor/Public/EnhancedInputEditorSettings.h
@@ -45,11 +45,11 @@ public:
 	 * If true then the Enhanced Input Editor subsystem will log all input that is being processed by it (keypresses, analog values, etc)
 	 * Note: This can produce A LOT of logs, so only use this if you are debugging something.
 	 */
-	UPROPERTY(config, EditAnywhere, Category = Logging, meta=(ConsoleVariable="EnhancedEditorInput.bShouldLogAllInputs"))
+	UPROPERTY(config, EditAnywhere, Category = Logging)
 	uint8 bLogAllInput : 1;
 	
 	/** If true, then the UEnhancedInputEditorSubsystem will be started when it is initalized */
-	UPROPERTY(config, EditAnywhere, Category = Editor, meta=(ConsoleVariable="EnhancedEditorInput.bAutomaticallyStartConsumingInput"))
+	UPROPERTY(config, EditAnywhere, Category = Editor)
 	uint8 bAutomaticallyStartConsumingInput : 1;
 
 	/** A bitmask of what event pins are visible when you place an Input Action event node in blueprints.  */
