Fix a input injection commands where, if you were doing it for an anlog key, it would not apply immediately. This was happening because the "Delta" value would not be applied for analog key , and we needed to force a re-evaluation of the event accumulator.

#rnx

diff --git a/Engine/Plugins/EnhancedInput/Source/EnhancedInput/Private/EnhancedInputSubsystemInterface.cpp b/Engine/Plugins/EnhancedInput/Source/EnhancedInput/Private/EnhancedInputSubsystemInterface.cpp
--- a/Engine/Plugins/EnhancedInput/Source/EnhancedInput/Private/EnhancedInputSubsystemInterface.cpp
+++ b/Engine/Plugins/EnhancedInput/Source/EnhancedInput/Private/EnhancedInputSubsystemInterface.cpp
@@ -12,6 +12,8 @@
 #include "InputMappingQuery.h"
 #include "PlayerMappableInputConfig.h"
 #include "PlayerMappableKeySettings.h"
+#include "Engine/LocalPlayer.h"
+#include "GenericPlatform/GenericPlatformInputDeviceMapper.h"
 
 #include UE_INLINE_GENERATED_CPP_BY_NAME(EnhancedInputSubsystemInterface)
 
@@ -1197,17 +1199,21 @@ void IEnhancedInputSubsystemInterface::TickForcedInput(float DeltaTime)
 	}
 
 	// Forced key presses
-	for (const TPair<FKey, FInputActionValue>& ForcedKeyPair : ForcedKeys)
+	for (TPair<FKey, FInjectedKeyData>& ForcedKeyPair : ForcedKeys)
 	{
 		// Prefer sending the key pressed event via a player controller if one is available.
 		if (APlayerController* Controller = Cast<APlayerController>(PlayerInput->GetOuter()))
 		{
-			InjectKey(Controller, ForcedKeyPair.Key, ForcedKeyPair.Value, DeltaTime);
+			InjectKey(Controller, ForcedKeyPair.Key, ForcedKeyPair.Value.InputValue, DeltaTime);
 		}
 		else
 		{
-			InjectKey(PlayerInput, ForcedKeyPair.Key, ForcedKeyPair.Value, DeltaTime);
+			InjectKey(PlayerInput, ForcedKeyPair.Key, ForcedKeyPair.Value.InputValue, DeltaTime);
 		}
+
+		// Keep track of the fact that we have injected this input value so we can check
+		// it if we remove input on the same frame
+		ForcedKeyPair.Value.LastInjectedValue = ForcedKeyPair.Value.InputValue;
 	}
 }
 
@@ -1230,7 +1236,9 @@ void IEnhancedInputSubsystemInterface::ApplyForcedInput(const UInputAction* Acti
 void IEnhancedInputSubsystemInterface::ApplyForcedInput(FKey Key, FInputActionValue Value)
 {
 	check(Key.IsValid());
-	ForcedKeys.Emplace(Key, Value);
+	
+	FInjectedKeyData& Data = ForcedKeys.FindOrAdd(Key);
+	Data.InputValue = Value;
 }
 
 void IEnhancedInputSubsystemInterface::RemoveForcedInput(const UInputAction* Action)
@@ -1241,15 +1249,34 @@ void IEnhancedInputSubsystemInterface::RemoveForcedInput(const UInputAction* Act
 void IEnhancedInputSubsystemInterface::RemoveForcedInput(FKey Key)
 {
 	check(Key.IsValid());
-	ForcedKeys.Remove(Key);
 
+	const FInjectedKeyData* InjectedKeyData = ForcedKeys.Find(Key);
+	if (!InjectedKeyData)
+	{
+		// Nothing to do if the value was not being injected
+		return;
+	}
+	
+	// Otherwise, we need to inject a release event tos player input
 	if (UEnhancedPlayerInput* PlayerInput = GetPlayerInput())
 	{
 		FInputKeyParams Params;
 		Params.Key = Key;
-		Params.Delta = FVector::ZeroVector;
+
+		// We want to inject the opposite of whatever we were previously injecting for this key
+		// in order to get it back to providing a fake value of zero. For example, if we were injecting (.5,.5)
+		// we want to use a delta of -.5,-.5 to get us back to a zero value. We only want to do this
+		// for analog keys.
+		//
+		// Any digital key we always want a value of zero to ensure it is treated as a release event.
+		Params.Delta = Key.IsAnalog() ? -InjectedKeyData->LastInjectedValue.Get<FVector>() : FVector::ZeroVector;
 		Params.Event = EInputEvent::IE_Released;
-		
+		Params.NumSamples = Key.IsAnalog() ? 1 : 0;
+
+		// Set the input device id to the platform user's default input device
+		const FPlatformUserId UserId = PlayerInput->GetOwningLocalPlayer()->GetPlatformUserId();
+		Params.InputDevice = IPlatformInputDeviceMapper::Get().GetPrimaryInputDeviceForUser(UserId);
+	
 		// Prefer sending the key released event via a player controller if one is available.
 		if (APlayerController* Controller = Cast<APlayerController>(PlayerInput->GetOuter()))
 		{
@@ -1259,5 +1286,12 @@ void IEnhancedInputSubsystemInterface::RemoveForcedInput(FKey Key)
 		{
 			PlayerInput->InputKey(Params);
 		}
+	
+		// Flush the player's pressed keys to ensure that the removed event is read
+		// and the PlayerInput re-evaluates the RawEventAccumulator as needed.
+		PlayerInput->FlushPressedKeys();
 	}
+
+	// No longer inject this key on tick
+	ForcedKeys.Remove(Key);
 }
diff --git a/Engine/Plugins/EnhancedInput/Source/EnhancedInput/Public/EnhancedInputSubsystemInterface.h b/Engine/Plugins/EnhancedInput/Source/EnhancedInput/Public/EnhancedInputSubsystemInterface.h
--- a/Engine/Plugins/EnhancedInput/Source/EnhancedInput/Public/EnhancedInputSubsystemInterface.h
+++ b/Engine/Plugins/EnhancedInput/Source/EnhancedInput/Public/EnhancedInputSubsystemInterface.h
@@ -520,7 +520,35 @@ private:
 	void ApplyAxisPropertyModifiers(UEnhancedPlayerInput* PlayerInput, struct FEnhancedActionKeyMapping& Mapping) const;
 
 	TMap<TWeakObjectPtr<const UInputAction>, FInputActionValue> ForcedActions;
-	TMap<FKey, FInputActionValue> ForcedKeys;
+
+	/** 
+	* Data storing the state of input values we inject for specific FKeys
+	*/
+	struct FInjectedKeyData
+	{
+		FInjectedKeyData() = default;
+		
+		FInjectedKeyData(const FInputActionValue& Value)
+			: InputValue(Value)
+			, LastInjectedValue(FInputActionValue(FVector::ZeroVector))
+		{}
+		
+		/** 
+		* The desired input action value to inject when we tick.
+		*/
+		FInputActionValue InputValue = FInputActionValue(FVector::ZeroVector);
+
+		/**
+		 * The last input value which was injected for this key. This is
+		 * set on tick after we have injected the input for this key at least once.
+		 */
+		FInputActionValue LastInjectedValue = FInputActionValue(FVector::ZeroVector);
+	};
+
+	/**
+	 * Map of FKey's to any input data that we are injecting to the player right now.
+	 */
+	TMap<FKey, FInjectedKeyData> ForcedKeys;
 
 	/**
 	 * A map of input actions with a Chorded trigger, mapped to the action they are dependent on.
