Add a "UnscaledRadial" deadzone option for games with simpler requirements, or for users who just prefer that type of deadzone.

This deadzone will block any input values if their magnitude is less then a lower threshold, and clamp the values to an upper threshold.

diff --git a/Engine/Plugins/EnhancedInput/Source/EnhancedInput/Private/InputModifiers.cpp b/Engine/Plugins/EnhancedInput/Source/EnhancedInput/Private/InputModifiers.cpp
--- a/Engine/Plugins/EnhancedInput/Source/EnhancedInput/Private/InputModifiers.cpp
+++ b/Engine/Plugins/EnhancedInput/Source/EnhancedInput/Private/InputModifiers.cpp
@@ -157,18 +157,29 @@ void UInputModifierDeadZone::PostEditChangeProperty(FPropertyChangedEvent& Prope
 
 FInputActionValue UInputModifierDeadZone::ModifyRaw_Implementation(const UEnhancedPlayerInput* PlayerInput, FInputActionValue CurrentValue, float DeltaTime)
 {
+	// Can't apply a deadzone to a boolean type (0 or 1 are the only options) 
 	EInputActionValueType ValueType = CurrentValue.GetValueType();
 	if (ValueType == EInputActionValueType::Boolean)
 	{
 		return CurrentValue;
 	}
-
-	auto DeadZoneLambda = [this](const float AxisVal)
+	
+	auto DeadZoneLambda = [this](const float AxisVal) -> float
 	{
 		// We need to translate and scale the input to the +/- 1 range after removing the dead zone.
 		return FMath::Min(1.f, (FMath::Max(0.f, FMath::Abs(AxisVal) - LowerThreshold) / (UpperThreshold - LowerThreshold))) * FMath::Sign(AxisVal);
 	};
 
+	auto UnscaledDeadZoneLambda = [this](const float AxisVal)-> float
+	{
+		// If the value is less then our lower threshold, return zero
+		// otherwise, clamp the value to the upper threshold.
+		return
+			AxisVal < LowerThreshold ?
+				0.0f :
+				FMath::Min(AxisVal, UpperThreshold);
+	};
+	
 	FVector NewValue = CurrentValue.Get<FVector>();
 	switch (Type)
 	{
@@ -191,6 +202,20 @@ FInputActionValue UInputModifierDeadZone::ModifyRaw_Implementation(const UEnhanc
 			NewValue.X = DeadZoneLambda(NewValue.X);
 		}
 		break;
+	case EDeadZoneType::UnscaledRadial:
+		if (ValueType == EInputActionValueType::Axis3D)
+		{
+			NewValue = NewValue.GetSafeNormal() * UnscaledDeadZoneLambda(NewValue.Size());
+		}
+		else if (ValueType == EInputActionValueType::Axis2D)
+		{
+			NewValue = NewValue.GetSafeNormal2D() * UnscaledDeadZoneLambda(NewValue.Size2D());
+		}
+		else
+		{
+			NewValue.X = UnscaledDeadZoneLambda(NewValue.X);
+		}
+		break;
 	}
 
 	return NewValue;
diff --git a/Engine/Plugins/EnhancedInput/Source/EnhancedInput/Public/InputModifiers.h b/Engine/Plugins/EnhancedInput/Source/EnhancedInput/Public/InputModifiers.h
--- a/Engine/Plugins/EnhancedInput/Source/EnhancedInput/Public/InputModifiers.h
+++ b/Engine/Plugins/EnhancedInput/Source/EnhancedInput/Public/InputModifiers.h
@@ -114,10 +114,27 @@ protected:
 UENUM()
 enum class EDeadZoneType : uint8
 {
-	// Apply dead zone to axes individually. This will result in input being chamfered at the corners for 2d/3d axis inputs, and matches the original UE4 deadzone logic.
+	// Apply dead zone to axes individually. This will result in input being chamfered at the corners for
+	// 2d/3d axis inputs, and matches the original UE4 deadzone logic.
 	Axial,
-	// Apply dead zone logic to all axes simultaneously. This gives smooth input (circular/spherical coverage). On a 1d axis input this works identically to Axial.
-	Radial,
+	
+	// Apply dead zone logic to all axes simultaneously. This gives smooth input (circular/spherical coverage).
+	// On a 1d axis input this works identically to Axial.
+	// 
+	// For most games, this will give the smoothest feeling analog values. The input is smoothed to avoid
+	// "jumpiness" when you are moving the analog axis.
+	Radial UMETA(DisplayName="Smoothed Radial"),
+
+	// Apply dead zone logic to all axes simultaneously without any smooth input
+	// which the normal "Radial" deadzone applies.
+	//
+	// The behavior of this deadzone type is as follows:
+	// If the magnitude of the input is less then the lower threshold, ignore it.
+	// Clamp the magnitude of the input to the upper threshold value.
+	//
+	// For some games, this may result in feeling "jumpy", because the value goes from 0.0 to
+	// the lower threshold immediately instead of being smoothed, like the normal "Radial" deadzone option.
+	UnscaledRadial
 };
 
 /** Dead Zone
