Use the widgets InputComponent instead of the player controllers when binding input events to it. This resolves an issue where widgets would always be binding inputs in enhanced input, but never actually unbinding them because the player controller input component was not being destroyed.

In order to use Enhanced Input event nodes in a widget, you must set the bAutomaticallyRegisterInputOnConstruction value to true on your widget. If you do not, then you will receive BP compilation errors.

You can use the "enhancedInput.bp.ShouldValidateWidgetBlueprintSettings" cvar to disable the BP errors if desired.

diff --git a/Engine/Plugins/EnhancedInput/EnhancedInput.uplugin b/Engine/Plugins/EnhancedInput/EnhancedInput.uplugin
--- a/Engine/Plugins/EnhancedInput/EnhancedInput.uplugin
+++ b/Engine/Plugins/EnhancedInput/EnhancedInput.uplugin
@@ -30,5 +30,11 @@
 			"Type": "Editor",
 			"LoadingPhase": "Default"
 		}
+	],
+	"Plugins": [
+		{
+			"Name": "DataValidation",
+			"Enabled": true
+		}
 	]
 }
\ No newline at end of file
diff --git a/Engine/Plugins/EnhancedInput/Source/InputBlueprintNodes/InputBlueprintNodes.Build.cs b/Engine/Plugins/EnhancedInput/Source/InputBlueprintNodes/InputBlueprintNodes.Build.cs
--- a/Engine/Plugins/EnhancedInput/Source/InputBlueprintNodes/InputBlueprintNodes.Build.cs
+++ b/Engine/Plugins/EnhancedInput/Source/InputBlueprintNodes/InputBlueprintNodes.Build.cs
@@ -22,6 +22,9 @@ namespace UnrealBuildTool.Rules
                     "Slate",
                     "SlateCore",
                     "UnrealEd",
+                    "UMGEditor",
+                    "UMG",
+                    "DataValidation"
                 }
             );
         }
diff --git a/Engine/Plugins/EnhancedInput/Source/InputBlueprintNodes/Private/EnhancedInputUserWidgetValidator.cpp b/Engine/Plugins/EnhancedInput/Source/InputBlueprintNodes/Private/EnhancedInputUserWidgetValidator.cpp
new file mode 100644
--- /dev/null
+++ b/Engine/Plugins/EnhancedInput/Source/InputBlueprintNodes/Private/EnhancedInputUserWidgetValidator.cpp
@@ -0,0 +1,91 @@
+// Copyright Epic Games, Inc. All Rights Reserved.
+
+#include "EnhancedInputUserWidgetValidator.h"
+
+#include "Blueprint/UserWidget.h"
+#include "EdGraph/EdGraph.h"
+#include "K2Node_EnhancedInputAction.h"
+#include "WidgetBlueprint.h"
+
+namespace UE::Input
+{
+	static bool bShouldValidateWidgetBlueprintSettings = true;
+	static FAutoConsoleVariableRef CVarShouldValidateWidgetBlueprintSettings(
+		TEXT("enhancedInput.bp.ShouldValidateWidgetBlueprintSettings"),
+		bShouldValidateWidgetBlueprintSettings,
+		TEXT("Should the Enhanced Input event node throw an error if a widget blueprint does not have bAutomaticallyRegisterInputOnConstruction set to true?"),
+		ECVF_Default);
+
+	/**
+	 * Returns true if the given widget BP has any nodes of the enhanced input type
+	 * in any of its graphs.
+	 */
+	static bool HasAnyEnhancedInputNodes(const UWidgetBlueprint* WidgetBP)
+	{
+		check(WidgetBP);
+		
+		bool bHasAnyEnhancedInputNodes = false;
+		
+		TArray<UEdGraph*> Graphs;
+		WidgetBP->GetAllGraphs(Graphs);
+		
+		for (const UEdGraph* Graph : Graphs)
+		{
+			TArray<UK2Node_EnhancedInputAction*> EventNodes;
+			Graph->GetNodesOfClass(EventNodes);
+
+			bHasAnyEnhancedInputNodes |= !EventNodes.IsEmpty();
+
+			if (bHasAnyEnhancedInputNodes)
+			{
+				break;
+			}
+		}
+
+		return bHasAnyEnhancedInputNodes;
+	}
+}
+
+bool UEnhancedInputUserWidgetValidator::CanValidateAsset_Implementation(const FAssetData& InAssetData, UObject* InAsset, FDataValidationContext& InContext) const
+{
+	return
+		UE::Input::bShouldValidateWidgetBlueprintSettings &&
+		InAsset &&
+		InAsset->IsA<UWidgetBlueprint>();
+}
+
+EDataValidationResult UEnhancedInputUserWidgetValidator::ValidateLoadedAsset_Implementation(const FAssetData& InAssetData, UObject* InAsset, FDataValidationContext& InContext)
+{
+	return CombineDataValidationResults(
+			Super::ValidateLoadedAsset_Implementation(InAssetData, InAsset, InContext),
+			ValidateWidgetBlueprint(InAssetData, Cast<const UWidgetBlueprint>(InAsset), InContext));
+}
+
+EDataValidationResult UEnhancedInputUserWidgetValidator::ValidateWidgetBlueprint(const FAssetData& InAssetData, const UWidgetBlueprint* WidgetBP, FDataValidationContext& InContext)
+{
+	// We need a valid CDO here to check the value of the bAutomaticallyRegisterInputOnConstruction setting
+	const UUserWidget* DefaultObj = WidgetBP->GeneratedClass ? Cast<const UUserWidget>(WidgetBP->GeneratedClass->GetDefaultObject(false)) : nullptr;
+	if (!DefaultObj)
+	{
+		return EDataValidationResult::NotValidated;
+	}
+
+	// If there are no Enhanced Input nodes then there is no need to check further
+	if (!UE::Input::HasAnyEnhancedInputNodes(WidgetBP))
+	{
+		return EDataValidationResult::Valid;
+	}
+
+	// If the widget has EI nodes and the bAutomaticallyRegisterInputOnConstruction is true, then it is valid...
+	if (DefaultObj->bAutomaticallyRegisterInputOnConstruction)
+	{
+		return EDataValidationResult::Valid;
+	}
+
+	// ... otherwise, this widget is invalid.
+	// Widgets require the bAutomaticallyRegisterInputOnConstruction to be true in order to receive enhanced input events. 
+	const FText CurrentError = NSLOCTEXT("EnhancedInput", "Input.Widget.Error", "'bAutomaticallyRegisterInputOnConstruction' must be true in order to use Enhanced Input in the widget");
+	
+	AssetMessage(InAssetData, EMessageSeverity::Error, CurrentError);
+	return EDataValidationResult::Invalid;	
+}
diff --git a/Engine/Plugins/EnhancedInput/Source/InputBlueprintNodes/Private/EnhancedInputUserWidgetValidator.h b/Engine/Plugins/EnhancedInput/Source/InputBlueprintNodes/Private/EnhancedInputUserWidgetValidator.h
new file mode 100644
--- /dev/null
+++ b/Engine/Plugins/EnhancedInput/Source/InputBlueprintNodes/Private/EnhancedInputUserWidgetValidator.h
@@ -0,0 +1,37 @@
+// Copyright Epic Games, Inc. All Rights Reserved.
+
+#pragma once
+
+#include "EditorValidatorBase.h"
+
+#include "EnhancedInputUserWidgetValidator.generated.h"
+
+class IAssetRegistry;
+class UWidgetBlueprint;
+
+/**
+ * Validates Widget Blueprints that have any Enhanced Input
+ * nodes in them to ensure that they have the correct "bAutomaticallyRegisterInputOnConstruction"
+ * setting value.
+ *
+ * Widgets require bAutomaticallyRegisterInputOnConstruction to be true in order to
+ * receive callbacks from Enhanced Input.
+ */
+UCLASS()
+class UEnhancedInputUserWidgetValidator : public UEditorValidatorBase
+{
+	GENERATED_BODY()
+
+public:
+	
+	UEnhancedInputUserWidgetValidator() = default;
+
+protected:
+	//~UEditorValidatorBase interface
+	virtual bool CanValidateAsset_Implementation(const FAssetData& InAssetData, UObject* InAsset, FDataValidationContext& InContext) const override;
+	virtual EDataValidationResult ValidateLoadedAsset_Implementation(const FAssetData& InAssetData, UObject* InAsset, FDataValidationContext& InContext) override;
+	//~End of UEditorValidatorBase interface
+
+private:
+	EDataValidationResult ValidateWidgetBlueprint(const FAssetData& InAssetData, const UWidgetBlueprint* InAsset, FDataValidationContext& InContext);
+};
\ No newline at end of file
diff --git a/Engine/Source/Runtime/UMG/Private/UserWidget.cpp b/Engine/Source/Runtime/UMG/Private/UserWidget.cpp
--- a/Engine/Source/Runtime/UMG/Private/UserWidget.cpp
+++ b/Engine/Source/Runtime/UMG/Private/UserWidget.cpp
@@ -76,6 +76,7 @@ FPaintContext::FPaintContext()
 // UUserWidget
 UUserWidget::UUserWidget(const FObjectInitializer& ObjectInitializer)
 	: Super(ObjectInitializer)
+	, bAutomaticallyRegisterInputOnConstruction(false)
 	, bHasScriptImplementedTick(true)
 	, bHasScriptImplementedPaint(true)
 	, bInitialized(false)
@@ -1712,9 +1713,19 @@ void UUserWidget::BindToAnimationEvent(UWidgetAnimation* InAnimation, FWidgetAni
 void UUserWidget::NativeOnInitialized()
 {
 	// Bind any input delegates that may be on this widget to its owning player controller
-	if(APlayerController* PC = GetOwningPlayer())
+	if (bAutomaticallyRegisterInputOnConstruction)
 	{
-		UInputDelegateBinding::BindInputDelegates(GetClass(), PC->InputComponent, this);		
+		// Only widgets with a valid player controller can bind to input delegates
+		if (GetOwningPlayer() != nullptr)
+		{
+			InitializeInputComponent();
+			check(InputComponent);
+			UInputDelegateBinding::BindInputDelegates(GetClass(), InputComponent, this);
+		}
+		else if (!IsEditorUtility())
+		{
+			UE_LOG(LogUMG, Error, TEXT("[%hs] Widget '%s' has bAutomaticallyRegisterInputOnConstruction set to true, but no valid player controller. Input delegates will not work!"), __func__, *GetNameSafe(this));
+		}
 	}
 	
 	if (UWidgetBlueprintGeneratedClass* BPClass = Cast<UWidgetBlueprintGeneratedClass>(GetClass()))
diff --git a/Engine/Source/Runtime/UMG/Public/Blueprint/UserWidget.h b/Engine/Source/Runtime/UMG/Public/Blueprint/UserWidget.h
--- a/Engine/Source/Runtime/UMG/Public/Blueprint/UserWidget.h
+++ b/Engine/Source/Runtime/UMG/Public/Blueprint/UserWidget.h
@@ -997,6 +997,13 @@ public:
 	UPROPERTY(EditAnywhere, BlueprintReadWrite, Getter = "IsInputActionBlocking", Setter = "SetInputActionBlocking", BlueprintSetter = "SetInputActionBlocking", Category = "Input")
 	uint8 bStopAction : 1;
 
+	/**
+	 * If true, this widget will automatically register its own input component upon construction.
+	 * This will allow the use of binding input delegates in the event graph.
+	 */
+	UPROPERTY(EditAnywhere, BlueprintReadOnly, Category = "Input")
+	uint8 bAutomaticallyRegisterInputOnConstruction : 1;
+
 public:
 
 	/**
