Remove the "Experimental" tag from player mappable key types. These are heavily in use and are no longer experimental

#rnx

diff --git a/Engine/Plugins/EnhancedInput/Source/EnhancedInput/Public/EnhancedInputDeveloperSettings.h b/Engine/Plugins/EnhancedInput/Source/EnhancedInput/Public/EnhancedInputDeveloperSettings.h
--- a/Engine/Plugins/EnhancedInput/Source/EnhancedInput/Public/EnhancedInputDeveloperSettings.h
+++ b/Engine/Plugins/EnhancedInput/Source/EnhancedInput/Public/EnhancedInputDeveloperSettings.h
@@ -70,16 +70,12 @@ public:
 	 * An instance of this class will be spawned by each Enhanced Input subsytem as a place to store
 	 * user settings such as keymappings, accessibility settings, etc. Subclass this to add more custom
 	 * options to your game.
-	 *
-	 * Note: This is a new experimental feature!
 	 */
 	UPROPERTY(config, EditAnywhere, NoClear, Category = "Enhanced Input|User Settings", meta=(editCondition = "bEnableUserSettings"))
 	TSoftClassPtr<UEnhancedInputUserSettings> UserSettingsClass;
 
 	/**
 	 * The default class for the player mappable key profile, used to store the key mappings set by the player in the user settings.
-	 * 
-	 * Note: This is a new experimental feature!
 	 */
 	UPROPERTY(config, EditAnywhere, NoClear, Category = "Enhanced Input|User Settings", meta=(editCondition = "bEnableUserSettings"))
 	TSoftClassPtr<UEnhancedPlayerMappableKeyProfile> DefaultPlayerMappableKeyProfileClass;
@@ -99,10 +95,8 @@ public:
 
 	/**
 	 * If true, then an instance of the User Settings Class will be created on each Enhanced Input subsystem.
-	 * 
-	 * Note: This is a new experimental feature!
 	 */
-	UPROPERTY(config, EditAnywhere, Category = "Enhanced Input|User Settings", meta=(DisplayName="Enable User Settings (Experimental)", DisplayPriority = 1))
+	UPROPERTY(config, EditAnywhere, Category = "Enhanced Input|User Settings", meta=(DisplayName="Enable User Settings", DisplayPriority = 1))
 	uint8 bEnableUserSettings : 1;
 
 	/** If true, then the DefaultMappingContexts will be applied to all Enhanced Input Subsystems. */
diff --git a/Engine/Plugins/EnhancedInput/Source/EnhancedInput/Public/PlayerMappableKeySettings.h b/Engine/Plugins/EnhancedInput/Source/EnhancedInput/Public/PlayerMappableKeySettings.h
--- a/Engine/Plugins/EnhancedInput/Source/EnhancedInput/Public/PlayerMappableKeySettings.h
+++ b/Engine/Plugins/EnhancedInput/Source/EnhancedInput/Public/PlayerMappableKeySettings.h
@@ -10,9 +10,8 @@ struct FEnhancedActionKeyMapping;
 
 /**
 * Hold setting information of an Action Input or a Action Key Mapping for setting screen and save purposes.
-* Experimental: Do not count on long term support for this structure.
 */
-UCLASS(DefaultToInstanced, EditInlineNew, DisplayName="Player Mappable Key Settings (Experimental)")
+UCLASS(DefaultToInstanced, EditInlineNew, DisplayName="Player Mappable Key Settings")
 class ENHANCEDINPUT_API UPlayerMappableKeySettings : public UObject
 {
 	GENERATED_BODY()
@@ -23,7 +22,7 @@ public:
 	virtual FName GetMappingName() const { return Name; }
 
 #if WITH_EDITOR
-	EDataValidationResult IsDataValid(class FDataValidationContext& Context) const;
+	EDataValidationResult IsDataValid(class FDataValidationContext& Context) const override;
 
 	/**
 	 * Get the known mapping names that are current in use. This is a helper function if you want to use a "GetOptions" metadata on a UPROPERTY.
diff --git a/Engine/Plugins/EnhancedInput/Source/EnhancedInput/Public/UserSettings/EnhancedInputUserSettings.h b/Engine/Plugins/EnhancedInput/Source/EnhancedInput/Public/UserSettings/EnhancedInputUserSettings.h
--- a/Engine/Plugins/EnhancedInput/Source/EnhancedInput/Public/UserSettings/EnhancedInputUserSettings.h
+++ b/Engine/Plugins/EnhancedInput/Source/EnhancedInput/Public/UserSettings/EnhancedInputUserSettings.h
@@ -467,7 +467,7 @@ struct ENHANCEDINPUT_API FPlayerMappableKeyProfileCreationArgs final
  * To customize this for your game, you can create a subclass of it and change the "UserSettingsClass" in the
  * Enhanced Input Project Settings.
  */
-UCLASS(config=GameUserSettings, DisplayName="Enhanced Input User Settings (Experimental)", Category="Enhanced Input|User Settings")
+UCLASS(config=GameUserSettings, DisplayName="Enhanced Input User Settings", Category="Enhanced Input|User Settings")
 class ENHANCEDINPUT_API UEnhancedInputUserSettings : public USaveGame
 {
 	GENERATED_BODY()
