Automatically set bAutomaticallyRegisterInputOnConstruction on widgets during compilation if there are any enhanced input nodes connected

#rnx

diff --git a/Engine/Plugins/EnhancedInput/Source/InputBlueprintNodes/Private/K2Node_EnhancedInputAction.cpp b/Engine/Plugins/EnhancedInput/Source/InputBlueprintNodes/Private/K2Node_EnhancedInputAction.cpp
--- a/Engine/Plugins/EnhancedInput/Source/InputBlueprintNodes/Private/K2Node_EnhancedInputAction.cpp
+++ b/Engine/Plugins/EnhancedInput/Source/InputBlueprintNodes/Private/K2Node_EnhancedInputAction.cpp
@@ -21,6 +21,8 @@
 #include "Modules/ModuleManager.h"
 #include "Styling/AppStyle.h"
 #include "EnhancedInputEditorSettings.h"
+#include "WidgetBlueprint.h"
+#include "Blueprint/UserWidget.h"
 
 #include UE_INLINE_GENERATED_CPP_BY_NAME(K2Node_EnhancedInputAction)
 
@@ -332,6 +334,16 @@ void UK2Node_EnhancedInputAction::ExpandNode(FKismetCompilerContext& CompilerCon
 		return InputActionEvent;
 	};
 
+	// Widget blueprints require the bAutomaticallyRegisterInputOnConstruction to be set to true in order to receive callbacks
+	if (GetBlueprint()->IsA<UWidgetBlueprint>())
+	{
+		CompilerContext.AddPostCDOCompiledStep([](const UObject::FPostCDOCompiledContext& Context, UObject* NewCDO)
+		{
+			UUserWidget* Widget = CastChecked<UUserWidget>(NewCDO);
+			Widget->bAutomaticallyRegisterInputOnConstruction = true;
+		});	
+	}
+
 	// Create temporary variables to copy ActionValue and ElapsedSeconds into
 	UK2Node_TemporaryVariable* ActionValueVar = CompilerContext.SpawnIntermediateNode<UK2Node_TemporaryVariable>(this, SourceGraph);
 	ActionValueVar->VariableType.PinCategory = UK2Node_GetInputActionValue::GetValueCategory(InputAction);
diff --git a/Engine/Source/Editor/KismetCompiler/Private/KismetCompiler.cpp b/Engine/Source/Editor/KismetCompiler/Private/KismetCompiler.cpp
--- a/Engine/Source/Editor/KismetCompiler/Private/KismetCompiler.cpp
+++ b/Engine/Source/Editor/KismetCompiler/Private/KismetCompiler.cpp
@@ -5437,13 +5437,24 @@ void FKismetCompilerContext::PostCDOCompiled(const UObject::FPostCDOCompiledCont
 		PRAGMA_DISABLE_DEPRECATION_WARNINGS;
 		NewClass->ClassDefaultObject->PostCDOCompiled();
 		PRAGMA_ENABLE_DEPRECATION_WARNINGS;
+		
+		for (const TFunctionRef<void(const UObject::FPostCDOCompiledContext&, UObject*)>& Step : PostCDOCompileSteps)
+		{
+			Step(Context, NewClass->ClassDefaultObject);
+		}
 	}
+	
 	FCoreUObjectDelegates::OnObjectPostCDOCompiled.Broadcast(NewClass->ClassDefaultObject, Context);
 
 	// Allow children to customize PostCDOCompile:
 	OnPostCDOCompiled(Context);
 }
 
+void FKismetCompilerContext::AddPostCDOCompiledStep(TFunctionRef<void(const UObject::FPostCDOCompiledContext&, UObject*)> StepFunction)
+{
+	PostCDOCompileSteps.Add(StepFunction);
+}
+
 void FKismetCompilerContext::Compile()
 {
 	CompileClassLayout(EInternalCompilerFlags::None);
diff --git a/Engine/Source/Editor/KismetCompiler/Public/KismetCompiler.h b/Engine/Source/Editor/KismetCompiler/Public/KismetCompiler.h
--- a/Engine/Source/Editor/KismetCompiler/Public/KismetCompiler.h
+++ b/Engine/Source/Editor/KismetCompiler/Public/KismetCompiler.h
@@ -88,6 +88,9 @@ protected:
 	// Map from node class to a handler functor
 	TMap< TSubclassOf<class UEdGraphNode>, FNodeHandlingFunctor*> NodeHandlers;
 
+	// Array of function refs to run after the CDO has been compiled in FKismetCompilerContext::PostCDOCompiled	
+	TArray<TFunctionRef<void(const UObject::FPostCDOCompiledContext&, UObject*)>> PostCDOCompileSteps;
+
 	// Map of properties created for timelines; to aid in debug data generation
 	TMap<class UTimelineTemplate*, class FProperty*> TimelineToMemberVariableMap;
 
@@ -205,6 +208,11 @@ public:
 	/** Called after the CDO has been generated, allows assignment of cached/derived data: */
 	void PostCDOCompiled(const UObject::FPostCDOCompiledContext& Context);
 
+	/**
+	* Adds the given "StepFunction" to be run after the CDO has been compiled in FKismetCompilerContext::PostCDOCompiled.
+	*/
+	void AddPostCDOCompiledStep(TFunctionRef<void (const UObject::FPostCDOCompiledContext&, UObject*)> StepFunction);
+
 	/** Compile a blueprint into a class and a set of functions */
 	void Compile();
 
diff --git a/Engine/Source/Runtime/UMG/Public/Blueprint/UserWidget.h b/Engine/Source/Runtime/UMG/Public/Blueprint/UserWidget.h
--- a/Engine/Source/Runtime/UMG/Public/Blueprint/UserWidget.h
+++ b/Engine/Source/Runtime/UMG/Public/Blueprint/UserWidget.h
@@ -1000,8 +1000,10 @@ public:
 	/**
 	 * If true, this widget will automatically register its own input component upon construction.
 	 * This will allow the use of binding input delegates in the event graph.
+	 * 
+	 * This is set during the compilation of the widget blueprint. 
 	 */
-	UPROPERTY(EditAnywhere, BlueprintReadOnly, Category = "Input")
+	UPROPERTY()
 	uint8 bAutomaticallyRegisterInputOnConstruction : 1;
 
 public:
