Remove some deprecated functions from Enhnaced Input subsystem interface from 5.2 and 5.3

#rnx

diff --git a/Engine/Plugins/EnhancedInput/Source/EnhancedInput/Private/EnhancedInputSubsystemInterface.cpp b/Engine/Plugins/EnhancedInput/Source/EnhancedInput/Private/EnhancedInputSubsystemInterface.cpp
--- a/Engine/Plugins/EnhancedInput/Source/EnhancedInput/Private/EnhancedInputSubsystemInterface.cpp
+++ b/Engine/Plugins/EnhancedInput/Source/EnhancedInput/Private/EnhancedInputSubsystemInterface.cpp
@@ -594,192 +594,6 @@ TArray<FEnhancedActionKeyMapping> IEnhancedInputSubsystemInterface::GetAllPlayer
 	return PlayerMappableMappings;
 }
 
-PRAGMA_DISABLE_DEPRECATION_WARNINGS
-
-int32 IEnhancedInputSubsystemInterface::AddPlayerMappedKey(const FName MappingName, const FKey NewKey, const FModifyContextOptions& Options /*= FModifyContextOptions()*/)
-{
-	return AddPlayerMappedKeyInSlot(MappingName, NewKey, FPlayerMappableKeySlot::FirstKeySlot, Options);
-}
-
-int32 IEnhancedInputSubsystemInterface::K2_AddPlayerMappedKeyInSlot(const FName MappingName, const FKey NewKey, const FPlayerMappableKeySlot& KeySlot /*= FPlayerMappableKeySlot()*/, const FModifyContextOptions& Options /*= FModifyContextOptions()*/)
-{
-	return AddPlayerMappedKeyInSlot(MappingName, NewKey, KeySlot, Options);
-}
-
-int32 IEnhancedInputSubsystemInterface::AddPlayerMappedKeyInSlot(const FName MappingName, const FKey NewKey, const FPlayerMappableKeySlot& KeySlot /*= FPlayerMappableKeySlot::FirstKeySlot*/, const FModifyContextOptions& Options /*= FModifyContextOptions()*/)
-{
-	if (UEnhancedInputUserSettings* Settings = GetUserSettings())
-	{
-		FMapPlayerKeyArgs Args = {};
-		Args.MappingName = MappingName;
-		Args.NewKey = NewKey;
-		Args.Slot = static_cast<EPlayerMappableKeySlot>(KeySlot.GetSlotNumber());
-
-		FGameplayTagContainer FailureReason;
-		Settings->MapPlayerKey(Args, FailureReason);
-
-		if (!FailureReason.IsEmpty())
-		{
-			UE_LOG(LogEnhancedInput, Warning, TEXT("IEnhancedInputSubsystemInterface::AddPlayerMappedKeyInSlot Failed! Reasoning: %s"), *FailureReason.ToString());
-			return 0;
-		}
-	}
-
-	RequestRebuildControlMappings(Options);
-	return 1;
-}
-
-int32 IEnhancedInputSubsystemInterface::RemovePlayerMappedKey(const FName MappingName, const FModifyContextOptions& Options /*= FModifyContextOptions()*/)
-{
-	return RemovePlayerMappedKeyInSlot(MappingName, FPlayerMappableKeySlot::FirstKeySlot, Options);
-}
-
-int32 IEnhancedInputSubsystemInterface::K2_RemovePlayerMappedKeyInSlot(const FName MappingName, const FPlayerMappableKeySlot& KeySlot /*= FPlayerMappableKeySlot()*/, const FModifyContextOptions& Options /*= FModifyContextOptions()*/)
-{
-	return RemovePlayerMappedKeyInSlot(MappingName, KeySlot, Options);
-}
-
-int32 IEnhancedInputSubsystemInterface::RemovePlayerMappedKeyInSlot(const FName MappingName, const FPlayerMappableKeySlot& KeySlot /*= FPlayerMappableKeySlot::FirstKeySlot*/, const FModifyContextOptions& Options /*= FModifyContextOptions()*/)
-{
-	if (UEnhancedInputUserSettings* Settings = GetUserSettings())
-	{
-		FMapPlayerKeyArgs Args = {};
-		Args.MappingName = MappingName;
-		Args.Slot = static_cast<EPlayerMappableKeySlot>(KeySlot.GetSlotNumber());
-		
-		FGameplayTagContainer FailureReason;
-		Settings->UnMapPlayerKey(Args, FailureReason);
-
-		if (!FailureReason.IsEmpty())
-		{
-			UE_LOG(LogEnhancedInput, Warning, TEXT("IEnhancedInputSubsystemInterface::RemovePlayerMappedKeyInSlot Failed! Reasoning: %s"), *FailureReason.ToString());
-			return 0;
-		}
-	}
-	
-	RequestRebuildControlMappings(Options);
-
-	return 1;
-}
-
-int32 IEnhancedInputSubsystemInterface::RemoveAllPlayerMappedKeysForMapping(const FName MappingName, const FModifyContextOptions& Options /*= FModifyContextOptions()*/)
-{
-	if (UEnhancedInputUserSettings* Settings = GetUserSettings())
-	{
-		FMapPlayerKeyArgs Args = {};
-		Args.MappingName = MappingName;
-
-		FGameplayTagContainer FailureReason;
-		Settings->ResetAllPlayerKeysInRow(Args, FailureReason);
-		
-		if (!FailureReason.IsEmpty())
-		{
-			UE_LOG(LogEnhancedInput, Warning, TEXT("IEnhancedInputSubsystemInterface::RemovePlayerMappedKeyInSlot Failed! Reasoning: %s"), *FailureReason.ToString());
-			return 0;
-		}
-	}
-
-	RequestRebuildControlMappings(Options);
-
-	return 1;
-}
-
-void IEnhancedInputSubsystemInterface::RemoveAllPlayerMappedKeys(const FModifyContextOptions& Options)
-{
-	if (UEnhancedInputUserSettings* Settings = GetUserSettings())
-	{
-		FGameplayTagContainer FailureReason;
-		Settings->ResetKeyProfileToDefault(Settings->GetCurrentKeyProfileIdentifier(), FailureReason);
-	}
-
-	RequestRebuildControlMappings(Options);
-}
-
-FKey IEnhancedInputSubsystemInterface::GetPlayerMappedKey(const FName MappingName) const
-{
-	return GetPlayerMappedKeyInSlot(MappingName, FPlayerMappableKeySlot::FirstKeySlot);
-}
-
-FKey IEnhancedInputSubsystemInterface::K2_GetPlayerMappedKeyInSlot(const FName MappingName, const FPlayerMappableKeySlot& KeySlot /*= FPlayerMappableKeySlot()*/) const
-{
-	return GetPlayerMappedKeyInSlot(MappingName, KeySlot);
-}
-
-FKey IEnhancedInputSubsystemInterface::GetPlayerMappedKeyInSlot(const FName MappingName, const FPlayerMappableKeySlot& KeySlot /*= FPlayerMappableKeySlot()*/) const
-{
-	if (const UEnhancedInputUserSettings* Settings = GetUserSettings())
-	{
-		if (const UEnhancedPlayerMappableKeyProfile* KeyProfile = Settings->GetCurrentKeyProfile())
-		{
-			FPlayerMappableKeyQueryOptions Opts = {};
-			Opts.MappingName = MappingName;
-			Opts.SlotToMatch = static_cast<EPlayerMappableKeySlot>(KeySlot.GetSlotNumber());
-
-			TArray<FKey> Keys;
-			KeyProfile->QueryPlayerMappedKeys(Opts, OUT Keys);
-			
-			if (!Keys.IsEmpty())
-			{
-				return Keys[0];
-			}
-		}
-	}
-	
-	return EKeys::Invalid;
-}
-
-TArray<FKey> IEnhancedInputSubsystemInterface::GetAllPlayerMappedKeys(const FName MappingName) const
-{
-	TArray<FKey> PlayerMappedKeys;
-	
-	if (const UEnhancedInputUserSettings* Settings = GetUserSettings())
-	{
-		if (const UEnhancedPlayerMappableKeyProfile* KeyProfile = Settings->GetCurrentKeyProfile())
-		{
-			FPlayerMappableKeyQueryOptions Opts = {};
-			Opts.MappingName = MappingName;
-			
-			KeyProfile->QueryPlayerMappedKeys(Opts, OUT PlayerMappedKeys);
-		}
-	}
-	
-	return PlayerMappedKeys;
-}
-
-void IEnhancedInputSubsystemInterface::AddPlayerMappableConfig(const UPlayerMappableInputConfig* Config, const FModifyContextOptions& Options)
-{
-	if (Config)
-	{
-		if (GetDefault<UEnhancedInputDeveloperSettings>()->bLogOnDeprecatedConfigUsed && Config->IsDeprecated())
-		{
-			UE_LOG(LogEnhancedInput, Warning, TEXT("The Player Mappable Input Config '%s' is marked as deprecated, but is still being added!"), *Config->GetFName().ToString());
-		}
-
-		for (TPair<TObjectPtr<UInputMappingContext>, int32> Pair : Config->GetMappingContexts())
-		{
-			AddMappingContext(Pair.Key, Pair.Value, Options);
-		}	
-	}
-}
-
-void IEnhancedInputSubsystemInterface::RemovePlayerMappableConfig(const UPlayerMappableInputConfig* Config, const FModifyContextOptions& Options)
-{
-	if (Config)
-	{
-		if (GetDefault<UEnhancedInputDeveloperSettings>()->bLogOnDeprecatedConfigUsed && Config->IsDeprecated())
-		{
-			UE_LOG(LogEnhancedInput, Warning, TEXT("The Player Mappable Input Config '%s' is marked as deprecated, but is still being removed!"), *Config->GetFName().ToString());
-		}
-
-		for(TPair<TObjectPtr<UInputMappingContext>, int32> Pair : Config->GetMappingContexts())
-		{
-			RemoveMappingContext(Pair.Key, Options);
-		}	
-	}
-}
-
-PRAGMA_ENABLE_DEPRECATION_WARNINGS
-
 // TODO: This should be a delegate (along with InjectChordBlockers), moving chording out of the underlying subsystem and enabling implementation of custom mapping handlers.
 /**
  * Reorder the given UnordedMappings such that chording mappings > chorded mappings > everything else.
diff --git a/Engine/Plugins/EnhancedInput/Source/EnhancedInput/Public/EnhancedInputSubsystemInterface.h b/Engine/Plugins/EnhancedInput/Source/EnhancedInput/Public/EnhancedInputSubsystemInterface.h
--- a/Engine/Plugins/EnhancedInput/Source/EnhancedInput/Public/EnhancedInputSubsystemInterface.h
+++ b/Engine/Plugins/EnhancedInput/Source/EnhancedInput/Public/EnhancedInputSubsystemInterface.h
@@ -326,141 +326,6 @@ public:
 	UFUNCTION(BlueprintCallable, BlueprintCosmetic, Category = "Input|Mapping Queries")
 	virtual TArray<FEnhancedActionKeyMapping> GetAllPlayerMappableActionKeyMappings() const;
 	
-	/**
-	 * Emplace or replace any currently applied key in the first key slot for mapping of MappingName.
-	 * Requests a rebuild of the player mappings. 
-	 *
-	 * @return The number of mappings that have been replaced
-	 */
-	UE_DEPRECATED(5.2, "AddPlayerMappedKey has been deprecated, please use AddPlayerMappedKeyInSlot instead.")
-	virtual int32 AddPlayerMappedKey(const FName MappingName, const FKey NewKey, const FModifyContextOptions& Options = FModifyContextOptions());
-
-	PRAGMA_DISABLE_DEPRECATION_WARNINGS
-	/**
-	 * Emplace or replace any currently applied key in KeySlot for mapping of MappingName.
-	 * Requests a rebuild of the player mappings.
-	 *
-	 * @return The number of mappings that have been replaced
-	 */
-	UE_DEPRECATED(5.3, "K2_AddPlayerMappedKeyInSlot has been deprecated, please use UEnhancedInputUserSettings instead.")
-	UFUNCTION(BlueprintCallable, BlueprintCosmetic, Category = "Input|PlayerMappable", meta = (
-		DeprecatedFunction,
-		DeprecationMessage="AddPlayerMappedKeyInSlot has been deprecated, please use UEnhancedInputUserSettings instead.",
-		DisplayName = "Add Player Mapped Key In Slot", AutoCreateRefTerm = "KeySlot, Options"))
-	virtual int32 K2_AddPlayerMappedKeyInSlot(const FName MappingName, const FKey NewKey, const FPlayerMappableKeySlot& KeySlot = FPlayerMappableKeySlot(), const FModifyContextOptions& Options = FModifyContextOptions());
-
-	/**
-	 * Emplace or replace any currently applied key in KeySlot for mapping of MappingName.
-	 * Requests a rebuild of the player mappings.
-	 *
-	 * @return The number of mappings that have been replaced
-	 */
-	UE_DEPRECATED(5.3, "AddPlayerMappedKeyInSlot has been deprecated, please use UEnhancedInputUserSettings instead.")
-	virtual int32 AddPlayerMappedKeyInSlot(const FName MappingName, const FKey NewKey, const FPlayerMappableKeySlot& KeySlot = FPlayerMappableKeySlot::FirstKeySlot, const FModifyContextOptions& Options = FModifyContextOptions());
-
-	/**
-	 * Removes player mapped key in the first KeySlot for mapping of MappingName.
-	 * Requests a rebuild of the player mappings. 
-	 *
-	 * @return The number of mappings that have been removed
-	 */
-	UE_DEPRECATED(5.2, "RemovePlayerMappedKey has been deprecated, please use RemovePlayerMappedKeyInSlot instead.")
-	virtual int32 RemovePlayerMappedKey(const FName MappingName, const FModifyContextOptions& Options = FModifyContextOptions());
-
-	/**
-	 * Removes player mapped key in the KeySlot for mapping of MappingName.
-	 * Requests a rebuild of the player mappings.
-	 *
-	 * @return The number of mappings that have been removed
-	 */
-	UFUNCTION(BlueprintCallable, BlueprintCosmetic, Category = "Input|PlayerMappable", meta = (DeprecatedFunction,
-		DeprecationMessage="K2_RemovePlayerMappedKeyInSlot has been deprecated, please use UEnhancedInputUserSettings instead.",
-		DisplayName="Remove Player Mapped Key In Slot", AutoCreateRefTerm = "KeySlot, Options"))
-	virtual int32 K2_RemovePlayerMappedKeyInSlot(const FName MappingName, const FPlayerMappableKeySlot& KeySlot = FPlayerMappableKeySlot(), const FModifyContextOptions& Options = FModifyContextOptions());
-
-	/**
-	 * Removes player mapped key in the KeySlot for mapping of MappingName.
-	 * Requests a rebuild of the player mappings.
-	 *
-	 * @return The number of mappings that have been removed
-	 */
-	UE_DEPRECATED(5.3, "RemovePlayerMappedKeyInSlot has been deprecated, please use UEnhancedInputUserSettings instead.")
-	virtual int32 RemovePlayerMappedKeyInSlot(const FName MappingName, const FPlayerMappableKeySlot& KeySlot = FPlayerMappableKeySlot::FirstKeySlot, const FModifyContextOptions& Options = FModifyContextOptions());
-
-	/**
-	 * Removes all player mapped keys for mapping of MappingName.
-	 * Requests a rebuild of the player mappings.
-	 *
-	 * @return The number of mappings that have been removed
-	 */
-	UE_DEPRECATED(5.3, "RemovePlayerMappedKeyInSlot has been deprecated, please use UEnhancedInputUserSettings instead.")
-	UFUNCTION(BlueprintCallable, BlueprintCosmetic, Category = "Input|PlayerMappable", meta = (DeprecatedFunction,
-		DeprecationMessage="RemoveAllPlayerMappedKeysForMapping has been deprecated, please use UEnhancedInputUserSettings instead.",
-		AutoCreateRefTerm = "KeySlot, Options"))
-	virtual int32 RemoveAllPlayerMappedKeysForMapping(const FName MappingName, const FModifyContextOptions& Options = FModifyContextOptions());
-
-	/**
-	 * Get the player mapped key in first slot to the given mapping name. If there is not a player mapped key, then this will return
-	 * EKeys::Invalid.
-	 *
-	 * @param MappingName	The FName of the mapped key that would have been set with the AddPlayerMappedKey function.
-	 */
-	UE_DEPRECATED(5.2, "GetPlayerMappedKey has been deprecated, please use GetPlayerMappedKeyInSlot instead.")
-	virtual FKey GetPlayerMappedKey(const FName MappingName) const;
-
-	/**
-	 * Get the player mapped key in first slot to the given mapping name. If there is not a player mapped key, then this will return
-	 * EKeys::Invalid.
-	 *
-	 * @param MappingName	The FName of the mapped key that would have been set with the AddPlayerMappedKey function.
-	 */
-	UE_DEPRECATED(5.3, "RemovePlayerMappedKeyInSlot has been deprecated, please use UEnhancedInputUserSettings instead.")
-	UFUNCTION(BlueprintCallable, BlueprintCosmetic, Category = "Input|PlayerMappable", meta = (
-		DeprecatedFunction,
-		DeprecationMessage="K2_GetPlayerMappedKeyInSlot has been deprecated, please use UEnhancedInputUserSettings instead.",
-		DisplayName="Get Player Mapped Key In Slot",AutoCreateRefTerm = "KeySlot"))
-	virtual FKey K2_GetPlayerMappedKeyInSlot(const FName MappingName, const FPlayerMappableKeySlot& KeySlot = FPlayerMappableKeySlot()) const;
-
-	/**
-	 * Get the player mapped key in first slot to the given mapping name. If there is not a player mapped key, then this will return
-	 * EKeys::Invalid.
-	 *
-	 * @param MappingName	The FName of the mapped key that would have been set with the AddPlayerMappedKey function.
-	 */
-	UE_DEPRECATED(5.3, "RemovePlayerMappedKeyInSlot has been deprecated, please use UEnhancedInputUserSettings instead.")
-	virtual FKey GetPlayerMappedKeyInSlot(const FName MappingName, const FPlayerMappableKeySlot& KeySlot = FPlayerMappableKeySlot::FirstKeySlot) const;
-
-	/**
-	 * Get all the player mapped keys to the given mapping name. If there is not a player mapped key, then this will return
-	 * TArray<FKey>().
-	 *
-	 * @param MappingName	The FName of the mapped key that would have been set with the AddPlayerMappedKey function.
-	 */
-	UE_DEPRECATED(5.3, "GetAllPlayerMappedKeys has been deprecated, please use UEnhancedInputUserSettings instead.")
-	UFUNCTION(BlueprintCallable, BlueprintCosmetic, Category = "Input|PlayerMappable", meta=(DeprecatedFunction, DeprecationMessage="GetAllPlayerMappedKeys has been deprecated, please use UEnhancedInputUserSettings instead."))
-	virtual TArray<FKey> GetAllPlayerMappedKeys(const FName MappingName) const;
-	
-	/**
-	 * Remove All PlayerMappedKeys
-	 * Requests a rebuild of the player mappings. 
-	 *
-	 */
-	UE_DEPRECATED(5.3, "RemoveAllPlayerMappedKeys has been deprecated, please use UEnhancedInputUserSettings instead.")
-	UFUNCTION(BlueprintCallable, BlueprintCosmetic, Category = "Input|PlayerMappable", meta=(AutoCreateRefTerm = "Options", DeprecatedFunction, DeprecationMessage="RemoveAllPlayerMappedKeys has been deprecated, please use UEnhancedInputUserSettings instead."))
-	virtual void RemoveAllPlayerMappedKeys(const FModifyContextOptions& Options = FModifyContextOptions());
-	
-	/** Adds all the input mapping contexts inside of this mappable config. */
-	UE_DEPRECATED(5.3, "RemoveAllPlayerMappedKeys has been deprecated, please use UEnhancedInputUserSettings instead.")
-	UFUNCTION(BlueprintCallable, BlueprintCosmetic, Category = "Input|PlayerMappable", meta=(AutoCreateRefTerm = "Options", DeprecatedFunction, DeprecationMessage="AddPlayerMappableConfig has been deprecated, please use UEnhancedInputUserSettings instead."))
-	virtual void AddPlayerMappableConfig(const UPlayerMappableInputConfig* Config, const FModifyContextOptions& Options = FModifyContextOptions());
-
-	/** Removes all the input mapping contexts inside of this mappable config. */
-	UE_DEPRECATED(5.3, "RemoveAllPlayerMappedKeys has been deprecated, please use UEnhancedInputUserSettings instead.")
-	UFUNCTION(BlueprintCallable, BlueprintCosmetic, Category = "Input|PlayerMappable", meta=(AutoCreateRefTerm = "Options", DeprecatedFunction, DeprecationMessage="RemovePlayerMappableConfig has been deprecated, please use UEnhancedInputUserSettings instead."))
-	virtual void RemovePlayerMappableConfig(const UPlayerMappableInputConfig* Config, const FModifyContextOptions& Options = FModifyContextOptions());
-	
-	PRAGMA_ENABLE_DEPRECATION_WARNINGS
-	
 protected:
 
 	/**
