Add Priority of MappingContexts as additional Debug information in EnhancedInput

diff --git a/Engine/Plugins/EnhancedInput/Source/EnhancedInput/Private/EnhancedInputSubsystemInterfaceDebug.cpp b/Engine/Plugins/EnhancedInput/Source/EnhancedInput/Private/EnhancedInputSubsystemInterfaceDebug.cpp
--- a/Engine/Plugins/EnhancedInput/Source/EnhancedInput/Private/EnhancedInputSubsystemInterfaceDebug.cpp
+++ b/Engine/Plugins/EnhancedInput/Source/EnhancedInput/Private/EnhancedInputSubsystemInterfaceDebug.cpp
@@ -123,14 +123,14 @@ void IEnhancedInputSubsystemInterface::ShowMappingContextDebugInfo(UCanvas* Canv
 			if (const TObjectPtr<const UInputMappingContext>* RedirectedContextPtr = AppliedContextRedirects.Find(AppliedContext))
 			{
 				const UInputMappingContext* RedirectedContext = *RedirectedContextPtr;
-				DisplayDebugManager.DrawString(FString::Printf(TEXT("  Redirected Context: %s -> %s"), *AppliedContext->GetFName().ToString(), *RedirectedContext->GetFName().ToString()));
+				DisplayDebugManager.DrawString(FString::Printf(TEXT("  Redirected Context: %s -> %s | Priority: %d"), *AppliedContext->GetFName().ToString(), *RedirectedContext->GetFName().ToString(), ContextPair.Value));
 
 				// Change the current context that is being used to display the correct mappings
 				AppliedContext = RedirectedContext;
 			}
 			else
 			{
-				DisplayDebugManager.DrawString(FString::Printf(TEXT("  Context: %s"), *AppliedContext->GetFName().ToString()));
+				DisplayDebugManager.DrawString(FString::Printf(TEXT("  Context: %s | Priority: %d"), *AppliedContext->GetFName().ToString(), ContextPair.Value));
 			}
 
 			// Build a table of mappings per action
