Create a "Repeated Tap" Enhanced Input trigger.

This input trigger subclasses the single "Tap" trigger, which we can utilize to detect an individual tap of a button. When we detect a single tap has happened, then check if it is within range of our "Repeat delay" time. If it is, then consider it a double tap!

diff --git a/Engine/Plugins/EnhancedInput/Source/EnhancedInput/Private/InputTriggers.cpp b/Engine/Plugins/EnhancedInput/Source/EnhancedInput/Private/InputTriggers.cpp
--- a/Engine/Plugins/EnhancedInput/Source/EnhancedInput/Private/InputTriggers.cpp
+++ b/Engine/Plugins/EnhancedInput/Source/EnhancedInput/Private/InputTriggers.cpp
@@ -46,6 +46,18 @@ namespace UE::Input
 		Result.RemoveFromEnd(TEXT("|"));
 		return Result;
 	}
+
+	static FString LexToString(const ETriggerState State)
+	{
+		switch (State)
+		{
+		case ETriggerState::Triggered: return TEXT("Triggered");
+		case ETriggerState::Ongoing: return TEXT("Ongoing");
+		case ETriggerState::None: return TEXT("None");
+		}
+
+		return TEXT("Invalid");
+	}
 }
 
 
@@ -211,6 +223,81 @@ ETriggerState UInputTriggerTap::UpdateState_Implementation(const UEnhancedPlayer
 	return State;
 }
 
+ETriggerState UInputTriggerRepeatedTap::UpdateState_Implementation(const UEnhancedPlayerInput* PlayerInput, FInputActionValue ModifiedValue, float DeltaTime)
+{	
+	const float LastHeldDuration = HeldDuration;
+
+	// Updates the held duration.
+	// This will return "none" if the key is not currently actuated, and "ongoing" if the key is currently being held down.
+	ETriggerState State = Super::UpdateState_Implementation(PlayerInput, ModifiedValue, DeltaTime);
+	
+	const bool bHasSingleTap = 
+		IsActuated(LastValue) &&						// If the key was actuated last frame...
+		State == ETriggerState::None &&					// and it now is not...
+		LastHeldDuration < TapReleaseTimeThreshold;		// and the amount of time that they key was held down before it was released is within our "single tap" time threshold...	
+	
+	const double CurrentTime = FPlatformTime::Seconds();
+
+	// Is the amount of time which has passed within the acceptable time frame to be considered a "repeat" tap? 
+	const bool bIsWithinValidRepeatTimeRange = CurrentTime <= RepeatTime;
+
+	// If we have any previous taps that we are keeping track of, then consider our state to be ongoing 
+	// This will stop our state from swapping back to "None" upon the release of the previous tap, and then immediately going
+	// back to "ongoing" 
+	// Note: It is a known issue that this is currently not supported by Enhanced Input because, even if this trigger returns "Ongoing",
+	// if the bound FKey has a value of zero, it will still be treated as the "None" later on during processing. 
+	State = (NumberOfTapsSinceLastTrigger > 0 || (bHasSingleTap && IsActuated(LastValue))) ? ETriggerState::Ongoing : State;
+
+	// If too much time has passed to be within a valid repeat range, and we have had previous taps,
+	// then we need to cancel this trigger.
+	if (!bIsWithinValidRepeatTimeRange && NumberOfTapsSinceLastTrigger > 0)
+	{
+		State = ETriggerState::None;
+	}
+	// If the key is currently being held down for longer then is allowed for a single "tap", then we also need to cancel.
+	else if (HeldDuration >= TapReleaseTimeThreshold)
+	{
+		State = ETriggerState::None;
+	}
+	// Otherwise, if we have detected a single tap and are within the allowed time range for repeating a tap, increment
+	// our tap count and see if we should trigger.
+	else if (bHasSingleTap && bIsWithinValidRepeatTimeRange)
+	{
+		++NumberOfTapsSinceLastTrigger;
+		const bool bHasReachedRepeatThreshold = NumberOfTapsSinceLastTrigger >= (NumberOfTapsWhichTriggerRepeat - 1);
+
+		if (bHasReachedRepeatThreshold)
+		{
+			State = ETriggerState::Triggered;
+		}
+		else
+		{
+			State = ETriggerState::Ongoing;
+		}
+	}
+
+	// Keep track of the next time range which is acceptable for another repeat
+	if (bHasSingleTap)
+	{
+		RepeatTime = CurrentTime + RepeatDelay;
+	}
+
+	// If the trigger has been cancelled somehow, reset the number of taps we have
+	if (State == ETriggerState::None || State == ETriggerState::Triggered)
+	{
+		NumberOfTapsSinceLastTrigger = 0;
+	}
+	
+	UE_LOG(LogEnhancedInput, Verbose, TEXT("Repeated Tap InputState:  %d / %d taps :: %s"), NumberOfTapsSinceLastTrigger, NumberOfTapsWhichTriggerRepeat, *UE::Input::LexToString(State));
+	
+	return State;
+}
+
+FString UInputTriggerRepeatedTap::GetDebugState() const
+{
+	return FString::Printf(TEXT("Repeated Taps:%d/%d"), NumberOfTapsSinceLastTrigger, NumberOfTapsWhichTriggerRepeat);
+}
+
 ETriggerState UInputTriggerPulse::UpdateState_Implementation(const UEnhancedPlayerInput* PlayerInput, FInputActionValue ModifiedValue, float DeltaTime)
 {
 	// Update HeldDuration and derive base state
diff --git a/Engine/Plugins/EnhancedInput/Source/EnhancedInput/Public/InputTriggers.h b/Engine/Plugins/EnhancedInput/Source/EnhancedInput/Public/InputTriggers.h
--- a/Engine/Plugins/EnhancedInput/Source/EnhancedInput/Public/InputTriggers.h
+++ b/Engine/Plugins/EnhancedInput/Source/EnhancedInput/Public/InputTriggers.h
@@ -349,6 +349,63 @@ public:
 	float TapReleaseTimeThreshold = 0.2f;
 };
 
+/**
+ * Fires the "Triggered" event upon a repeated taps of a key.
+ * 
+ * Note: You can use this trigger to easily make a "Double Tap" effect.
+ */
+UCLASS(NotBlueprintable, MinimalAPI, meta = (DisplayName = "Repeated Tap"))
+class UInputTriggerRepeatedTap : public UInputTriggerTimedBase
+{
+	GENERATED_BODY()
+
+protected:
+	virtual ETriggerState UpdateState_Implementation(const UEnhancedPlayerInput* PlayerInput, FInputActionValue ModifiedValue, float DeltaTime) override;
+	virtual FString GetDebugState() const override;
+	
+	/**
+	 * The max amount of time which can pass in between single taps that can pass and still be considered
+	 * a repeat tap. For example:
+	 *
+	 * -> Single Key Tap (the first tap)
+	 * --
+	 * -- ... Some time passes
+	 * --
+	 * -> Another single key tap (the second tap)
+	 * 
+	 * -- if (time passed <= RepeatDelay) then trigger repeat tap
+	 */
+	UPROPERTY(EditAnywhere, Config, BlueprintReadWrite, Category="Trigger Settings", meta = (ClampMin = "0"))
+	double RepeatDelay = 0.5;
+
+	/**
+	 * This is the time at which we need to be <= on the next detection of a
+	 * tap in order to be considered a double tap.
+	 *
+	 * This is set to (FPlatform::Seconds() + RepeatDelay) whenever we detect a tap
+	 */
+	UPROPERTY()
+	double RepeatTime = 0.0;
+
+	/**
+	 * The number of taps in a row it takes for this to be considered triggered.
+	 *
+	 * Note: To create a "Double tap" trigger, this value should be 2.
+	 */
+	UPROPERTY(EditAnywhere, Config, BlueprintReadWrite, Category="Trigger Settings", meta = (ClampMin = "1"))
+	int32 NumberOfTapsWhichTriggerRepeat = 2;
+
+	// Release within this time-frame to trigger a tap
+	UPROPERTY(EditAnywhere, Config, BlueprintReadWrite, Category = "Trigger Settings", meta = (ClampMin = "0"))
+	float TapReleaseTimeThreshold = 0.2f;
+
+	/**
+	 * Tracks the number of individual taps that have occured since the last time we met
+	 */
+	UPROPERTY()
+	int32 NumberOfTapsSinceLastTrigger = 0;
+};
+
 /** UInputTriggerPulse
 	Trigger that fires at an Interval, in seconds, while input is actuated. 
 	Note:	Completed only fires when the repeat limit is reached or when input is released immediately after being triggered.
