PR #11264: EnhancedInput: Expose the UObject bound to FEnhancedInputActionEventBinding

diff --git a/Engine/Plugins/EnhancedInput/Source/EnhancedInput/Public/EnhancedInputComponent.h b/Engine/Plugins/EnhancedInput/Source/EnhancedInput/Public/EnhancedInputComponent.h
--- a/Engine/Plugins/EnhancedInput/Source/EnhancedInput/Public/EnhancedInputComponent.h
+++ b/Engine/Plugins/EnhancedInput/Source/EnhancedInput/Public/EnhancedInputComponent.h
@@ -43,6 +43,10 @@ protected:
 	bool bShouldFireWithEditorScriptGuard = false;
 
 public:
+	UObject* GetUObject() const
+	{
+		return IsBound() ? Delegate->GetUObject() : nullptr;
+	}
 	
 	bool IsBound() const
 	{
@@ -177,6 +181,7 @@ public:
 	virtual TUniquePtr<FEnhancedInputActionEventBinding> Clone() const = 0;
 	virtual void SetShouldFireWithEditorScriptGuard(const bool bNewValue) = 0;
 	virtual bool IsBoundToObject(void const* Object) const = 0;
+	virtual UObject* GetUObject() const { return nullptr; }
 };
 
 /** Binds an action value for later reference. CurrentValue will be kept up to date with the value of the bound action */
@@ -271,6 +276,11 @@ public:
 		return Delegate.IsBoundToObject(Object);
 	}
 
+	virtual UObject* GetUObject() const override
+	{
+		return Delegate.GetUObject();
+	}
+
 	TEnhancedInputUnifiedDelegate<TSignature> Delegate;
 };
 
