Fixing GC issues around uninstalling GamefeaturePlugins.
Adds function for UnregisterKeyMappingsToProfile to remove

Review-37241199

diff --git a/Engine/Plugins/EnhancedInput/Source/EnhancedInput/Private/UserSettings/EnhancedInputUserSettings.cpp b/Engine/Plugins/EnhancedInput/Source/EnhancedInput/Private/UserSettings/EnhancedInputUserSettings.cpp
--- a/Engine/Plugins/EnhancedInput/Source/EnhancedInput/Private/UserSettings/EnhancedInputUserSettings.cpp
+++ b/Engine/Plugins/EnhancedInput/Source/EnhancedInput/Private/UserSettings/EnhancedInputUserSettings.cpp
@@ -1414,6 +1414,28 @@ bool UEnhancedInputUserSettings::RegisterKeyMappingsToProfile(UEnhancedPlayerMap
 	return true;
 }
 
+bool UEnhancedInputUserSettings::UnregisterKeyMappingsToProfile(UEnhancedPlayerMappableKeyProfile& Profile, const UInputMappingContext* IMC)
+{
+	if (!IMC)
+	{
+		UE_LOG(LogEnhancedInput, Error, TEXT("Attempting to unregister a null mapping context with the user settings!"));
+		return false;
+	}
+
+	for (const FEnhancedActionKeyMapping& KeyMapping : IMC->GetMappings())
+	{
+		// Skip over non-player mappable keys
+		if (!KeyMapping.IsPlayerMappable())
+		{
+			continue;
+		}
+
+		const FName MappingName = KeyMapping.GetMappingName();
+		Profile.PlayerMappedKeys.Remove(MappingName);
+	}
+	return true;
+}
+
 void UEnhancedInputUserSettings::OnKeyMappingRegistered(FPlayerKeyMapping& RegisteredMapping, const FEnhancedActionKeyMapping& SourceMapping)
 {
 	// Does nothing by default
@@ -1426,7 +1448,13 @@ FHardwareDeviceIdentifier UEnhancedInputUserSettings::DetermineHardwareDeviceFor
 
 bool UEnhancedInputUserSettings::UnregisterInputMappingContext(const UInputMappingContext* IMC)
 {
-	return RegisteredMappingContexts.Remove(IMC) != INDEX_NONE;
+	bool bResult = true;
+	for (TPair<FGameplayTag, TObjectPtr<UEnhancedPlayerMappableKeyProfile>> Pair : SavedKeyProfiles)
+	{
+		bResult &= UnregisterKeyMappingsToProfile(*Pair.Value, IMC);
+	}
+	bResult &= RegisteredMappingContexts.Remove(IMC) != INDEX_NONE;
+	return bResult;
 }
 
 bool UEnhancedInputUserSettings::UnregisterInputMappingContexts(const TSet<UInputMappingContext*>& MappingContexts)
diff --git a/Engine/Plugins/EnhancedInput/Source/EnhancedInput/Public/UserSettings/EnhancedInputUserSettings.h b/Engine/Plugins/EnhancedInput/Source/EnhancedInput/Public/UserSettings/EnhancedInputUserSettings.h
--- a/Engine/Plugins/EnhancedInput/Source/EnhancedInput/Public/UserSettings/EnhancedInputUserSettings.h
+++ b/Engine/Plugins/EnhancedInput/Source/EnhancedInput/Public/UserSettings/EnhancedInputUserSettings.h
@@ -697,6 +697,8 @@ protected:
 	 */
 	virtual bool RegisterKeyMappingsToProfile(UEnhancedPlayerMappableKeyProfile& Profile, const UInputMappingContext* IMC);
 
+	virtual bool UnregisterKeyMappingsToProfile(UEnhancedPlayerMappableKeyProfile& Profile, const UInputMappingContext* IMC);
+
 	/**
 	 * Determines the hardware device that an action key mapping is associated with. By default, the hardware will be FHardwareDeviceIdentifier::Invalid.
 	 *
