Remove the old deprecated type FPlayerMappableKeySlot from 5.3. Use EPlayerMappableKeySlot instead.

#rnx

diff --git a/Engine/Plugins/EnhancedInput/Source/EnhancedInput/Private/PlayerMappableKeySlot.cpp b/Engine/Plugins/EnhancedInput/Source/EnhancedInput/Private/PlayerMappableKeySlot.cpp
deleted file mode 100644
diff --git a/Engine/Plugins/EnhancedInput/Source/EnhancedInput/Public/EnhancedInputLibrary.h b/Engine/Plugins/EnhancedInput/Source/EnhancedInput/Public/EnhancedInputLibrary.h
--- a/Engine/Plugins/EnhancedInput/Source/EnhancedInput/Public/EnhancedInputLibrary.h
+++ b/Engine/Plugins/EnhancedInput/Source/EnhancedInput/Public/EnhancedInputLibrary.h
@@ -3,7 +3,6 @@
 #pragma once
 
 #include "Kismet/BlueprintFunctionLibrary.h"
-#include "PlayerMappableKeySlot.h"
 
 #include "EnhancedInputLibrary.generated.h"
 
@@ -103,26 +102,6 @@ public:
 	UFUNCTION(BlueprintPure, meta = (DisplayName = "To String (ETriggerEvent)", CompactNodeTitle = "->", BlueprintAutocast), Category = "Utilities|String")
 	static FString Conv_TriggerEventValueToString(const ETriggerEvent TriggerEvent);
 
-	PRAGMA_DISABLE_DEPRECATION_WARNINGS
-
-	UE_DEPRECATED(5.3, "FPlayerMappableKeySlot has been deprecated. Please use EPlayerMappableKeySlot instead.")
-	UFUNCTION(BlueprintPure, Category = "Input", meta = (ReturnDisplayName = "First Player Mappable Key Slot", DeprecatedFunction, DeprecationMessage="FPlayerMappableKeyOptions has been deprecated. Please use UPlayerMappableKeySettings instead."))
-	static FPlayerMappableKeySlot& GetFirstPlayerMappableKeySlot() { return FPlayerMappableKeySlot::FirstKeySlot; };
-
-	UE_DEPRECATED(5.3, "FPlayerMappableKeySlot has been deprecated. Please use EPlayerMappableKeySlot instead.")
-	UFUNCTION(BlueprintPure, Category = "Input", meta = (ReturnDisplayName = "Second Player Mappable Key Slot", DeprecatedFunction, DeprecationMessage="FPlayerMappableKeyOptions has been deprecated. Please use UPlayerMappableKeySettings instead."))
-	static FPlayerMappableKeySlot& GetSecondPlayerMappableKeySlot() { return FPlayerMappableKeySlot::SecondKeySlot; };
-
-	UE_DEPRECATED(5.3, "FPlayerMappableKeySlot has been deprecated. Please use EPlayerMappableKeySlot instead.")
-	UFUNCTION(BlueprintPure, Category = "Input", meta = (ReturnDisplayName = "Third Player Mappable Key Slot", DeprecatedFunction, DeprecationMessage="FPlayerMappableKeyOptions has been deprecated. Please use UPlayerMappableKeySettings instead."))
-	static FPlayerMappableKeySlot& GetThirdPlayerMappableKeySlot() { return FPlayerMappableKeySlot::ThirdKeySlot; };
-
-	UE_DEPRECATED(5.3, "FPlayerMappableKeySlot has been deprecated. Please use EPlayerMappableKeySlot instead.")
-	UFUNCTION(BlueprintPure, Category = "Input", meta = (ReturnDisplayName = "Fourth Player Mappable Key Slot", DeprecatedFunction, DeprecationMessage="FPlayerMappableKeyOptions has been deprecated. Please use UPlayerMappableKeySettings instead."))
-	static FPlayerMappableKeySlot& GetFourthPlayerMappableKeySlot() { return FPlayerMappableKeySlot::FourthKeySlot; };
-	
-	PRAGMA_ENABLE_DEPRECATION_WARNINGS
-
 	/**
 	 * Flushes the player controller's pressed keys
 	 * 
diff --git a/Engine/Plugins/EnhancedInput/Source/EnhancedInput/Public/EnhancedInputSubsystemInterface.h b/Engine/Plugins/EnhancedInput/Source/EnhancedInput/Public/EnhancedInputSubsystemInterface.h
--- a/Engine/Plugins/EnhancedInput/Source/EnhancedInput/Public/EnhancedInputSubsystemInterface.h
+++ b/Engine/Plugins/EnhancedInput/Source/EnhancedInput/Public/EnhancedInputSubsystemInterface.h
@@ -4,7 +4,6 @@
 
 #include "UObject/Interface.h"
 #include "EnhancedPlayerInput.h"
-#include "PlayerMappableKeySlot.h"
 
 #include "EnhancedInputSubsystemInterface.generated.h"
 
diff --git a/Engine/Plugins/EnhancedInput/Source/EnhancedInput/Public/PlayerMappableKeySlot.h b/Engine/Plugins/EnhancedInput/Source/EnhancedInput/Public/PlayerMappableKeySlot.h
deleted file mode 100644
