Enhanced Input Modes! Users can now specify IMC assets to only be active during a certian input mode. An input mode is a gameplay tag container on the enhanced player input object, muteable through the EI Subsystem. Only mappings contexts with filtering enabled will be affected

diff --git a/Engine/Plugins/EnhancedInput/Source/EnhancedInput/Private/EnhancedInputDeveloperSettings.cpp b/Engine/Plugins/EnhancedInput/Source/EnhancedInput/Private/EnhancedInputDeveloperSettings.cpp
--- a/Engine/Plugins/EnhancedInput/Source/EnhancedInput/Private/EnhancedInputDeveloperSettings.cpp
+++ b/Engine/Plugins/EnhancedInput/Source/EnhancedInput/Private/EnhancedInputDeveloperSettings.cpp
@@ -30,6 +30,7 @@ UEnhancedInputDeveloperSettings::UEnhancedInputDeveloperSettings(const FObjectIn
 	, bEnableUserSettings(false)
 	, bEnableDefaultMappingContexts(true)
 	, bShouldOnlyTriggerLastActionInChord(true)
+	, bEnableInputModeFiltering(true)
 	, bLogOnDeprecatedConfigUsed(true)
 	, bEnableWorldSubsystem(false)
 	, bShouldLogAllWorldSubsystemInputs(false)
diff --git a/Engine/Plugins/EnhancedInput/Source/EnhancedInput/Private/EnhancedInputSubsystemInterface.cpp b/Engine/Plugins/EnhancedInput/Source/EnhancedInput/Private/EnhancedInputSubsystemInterface.cpp
--- a/Engine/Plugins/EnhancedInput/Source/EnhancedInput/Private/EnhancedInputSubsystemInterface.cpp
+++ b/Engine/Plugins/EnhancedInput/Source/EnhancedInput/Private/EnhancedInputSubsystemInterface.cpp
@@ -274,6 +274,86 @@ void IEnhancedInputSubsystemInterface::RemoveMappingContext(const UInputMappingC
 	}
 }
 
+FGameplayTagContainer IEnhancedInputSubsystemInterface::GetInputMode() const
+{
+	UE_CLOG(!GetDefault<UEnhancedInputDeveloperSettings>()->bEnableInputModeFiltering,		
+		LogEnhancedInput, Warning, TEXT("[%hs] bEnableInputModeFiltering is false in the Enhanced Input developer settings. Nothing will happen."),
+		__func__);
+	
+	if (const UEnhancedPlayerInput* Input = GetPlayerInput())
+	{
+		return Input->GetCurrentInputMode();
+	}
+
+	UE_LOG(LogEnhancedInput, Error, TEXT("[%hs] Null player input, cannot get the current input mode."), __func__);
+	
+	return FGameplayTagContainer {};
+}
+
+void IEnhancedInputSubsystemInterface::SetInputMode(const FGameplayTagContainer& NewMode, const FModifyContextOptions& Options)
+{
+	UE_CLOG(!GetDefault<UEnhancedInputDeveloperSettings>()->bEnableInputModeFiltering,
+		LogEnhancedInput, Warning, TEXT("[%hs] bEnableInputModeFiltering is false in the Enhanced Input developer settings. Nothing will happen."),
+		__func__);
+	
+	UEnhancedPlayerInput* Input = GetPlayerInput();
+	if (!Input)
+	{
+		UE_LOG(LogEnhancedInput, Error, TEXT("[%hs] Null player input, unable to set the input mode to '%s'"), __func__, *NewMode.ToString());
+		return;
+	}
+	
+	Input->SetCurrentInputMode(NewMode);
+
+	RequestRebuildControlMappings(Options);
+}
+
+void IEnhancedInputSubsystemInterface::AppendTagsToInputMode(const FGameplayTagContainer& TagsToAdd, const FModifyContextOptions& Options)
+{
+	UE_CLOG(!GetDefault<UEnhancedInputDeveloperSettings>()->bEnableInputModeFiltering,
+		LogEnhancedInput, Warning, TEXT("[%hs] bEnableInputModeFiltering is false in the Enhanced Input developer settings. Nothing will happen."),
+		__func__);
+	
+	UEnhancedPlayerInput* Input = GetPlayerInput();
+	if (!Input)
+	{
+		UE_LOG(LogEnhancedInput, Error, TEXT("[%hs] Null player input, unable to append tags '%s' to the input mode"), __func__, *TagsToAdd.ToString());
+		return;
+	}
+	
+	Input->GetCurrentInputMode().AppendTags(TagsToAdd);
+	
+	RequestRebuildControlMappings(Options);
+}
+
+void IEnhancedInputSubsystemInterface::AddTagToInputMode(const FGameplayTag& TagToAdd, const FModifyContextOptions& Options)
+{
+	AppendTagsToInputMode(FGameplayTagContainer { TagToAdd }, Options);
+}
+
+void IEnhancedInputSubsystemInterface::RemoveTagsFromInputMode(const FGameplayTagContainer& TagsToRemove, const FModifyContextOptions& Options)
+{
+	UE_CLOG(!GetDefault<UEnhancedInputDeveloperSettings>()->bEnableInputModeFiltering,
+		LogEnhancedInput, Warning, TEXT("[%hs] bEnableInputModeFiltering is false in the Enhanced Input developer settings. Nothing will happen."),
+		__func__);
+	
+	UEnhancedPlayerInput* Input = GetPlayerInput();
+	if (!Input)
+	{
+		UE_LOG(LogEnhancedInput, Error, TEXT("[%hs] Null player input, unable to remove tags '%s' from the input mode"), __func__, *TagsToRemove.ToString());
+		return;
+	}
+
+	Input->GetCurrentInputMode().RemoveTags(TagsToRemove);
+	
+	RequestRebuildControlMappings(Options);
+}
+
+void IEnhancedInputSubsystemInterface::RemoveTagFromInputMode(const FGameplayTag& TagToRemove, const FModifyContextOptions& Options)
+{
+	RemoveTagsFromInputMode(FGameplayTagContainer { TagToRemove }, Options);
+}
+
 void IEnhancedInputSubsystemInterface::RequestRebuildControlMappings(const FModifyContextOptions& Options, EInputMappingRebuildType MappingRebuildType)
 {
 	bMappingRebuildPending = true;
@@ -724,7 +804,11 @@ void IEnhancedInputSubsystemInterface::RebuildControlMappings()
 		// TODO: Prefer to reset MappingRebuildPending here?
 		return;
 	}
-
+	
+	const FGameplayTagContainer& CurrentInputMode = PlayerInput->GetCurrentInputMode();
+	const bool bInputModeFilteringEnabled = GetDefault<UEnhancedInputDeveloperSettings>()->bEnableInputModeFiltering;
+	const bool bIsCurrentInputModeEmpty = CurrentInputMode.IsEmpty();
+	
 	// Clear existing mappings, but retain the mapping array for later processing
 	TArray<FEnhancedActionKeyMapping> OldMappings(MoveTemp(PlayerInput->EnhancedActionMappings));
 	PlayerInput->ClearAllMappings();
@@ -743,7 +827,7 @@ void IEnhancedInputSubsystemInterface::RebuildControlMappings()
 		{
 			if (!Pair.Key || !Pair.Value)
 			{
-				UE_LOG(LogEnhancedInput, Error, TEXT("An invalid Mappping Context Redirect specified in '%s'"), PlatformSettings->GetConfigOverridePlatform());
+				UE_LOG(LogEnhancedInput, Error, TEXT("An invalid Mapping Context Redirect specified in '%s'"), PlatformSettings->GetConfigOverridePlatform());
 				continue;
 			}
 			
@@ -791,6 +875,19 @@ void IEnhancedInputSubsystemInterface::RebuildControlMappings()
 		// Don't apply context specific keys immediately, allowing multiple mappings to the same key within the same context if required.
 		TArray<FKey> ContextAppliedKeys;
 		const UInputMappingContext* MappingContext = ContextPair.Key;
+		
+		// Check if this mapping context can be applied in our current input mode on the player
+		// If it can't, then we will not process its mappings..
+		if (bInputModeFilteringEnabled &&
+			MappingContext->bFilterMappingsByQuery &&
+			!MappingContext->InputModeQuery.Matches(CurrentInputMode))
+		{
+			UE_LOG(LogEnhancedInput, Log, TEXT("[%hs] Not applying mappings from IMC '%s' because it does not meet the requirements of the current input mode '%s'"),
+				__func__, *GetNameSafe(MappingContext), *CurrentInputMode.ToString());
+				
+			continue;
+		}
+		
 		TArray<FEnhancedActionKeyMapping> OrderedMappings = ReorderMappings(MappingContext->GetMappings(), PlayerInput->DependentChordActions);
 
 		for (FEnhancedActionKeyMapping& Mapping : OrderedMappings)
diff --git a/Engine/Plugins/EnhancedInput/Source/EnhancedInput/Private/EnhancedPlayerInput.cpp b/Engine/Plugins/EnhancedInput/Source/EnhancedInput/Private/EnhancedPlayerInput.cpp
--- a/Engine/Plugins/EnhancedInput/Source/EnhancedInput/Private/EnhancedPlayerInput.cpp
+++ b/Engine/Plugins/EnhancedInput/Source/EnhancedInput/Private/EnhancedPlayerInput.cpp
@@ -277,6 +277,21 @@ float UEnhancedPlayerInput::GetEffectiveTimeDilation() const
 	return 1.0f;
 }
 
+FGameplayTagContainer& UEnhancedPlayerInput::GetCurrentInputMode()
+{
+	return CurrentInputMode;
+}
+
+const FGameplayTagContainer& UEnhancedPlayerInput::GetCurrentInputMode() const
+{
+	return CurrentInputMode;
+}
+
+void UEnhancedPlayerInput::SetCurrentInputMode(const FGameplayTagContainer& NewMode)
+{
+	CurrentInputMode = NewMode;
+}
+
 void UEnhancedPlayerInput::EvaluateKeyMapState(const float DeltaTime, const bool bGamePaused, OUT TArray<TPair<FKey, FKeyState*>>& KeysWithEvents)
 {
 	TRACE_CPUPROFILER_EVENT_SCOPE(EnhPIS_KeyDownPrev);
diff --git a/Engine/Plugins/EnhancedInput/Source/EnhancedInput/Private/InputMappingContext.cpp b/Engine/Plugins/EnhancedInput/Source/EnhancedInput/Private/InputMappingContext.cpp
--- a/Engine/Plugins/EnhancedInput/Source/EnhancedInput/Private/InputMappingContext.cpp
+++ b/Engine/Plugins/EnhancedInput/Source/EnhancedInput/Private/InputMappingContext.cpp
@@ -5,6 +5,7 @@
 #include "EnhancedInputLibrary.h"
 #include "EnhancedInputModule.h"
 #include "PlayerMappableKeySettings.h"
+#include "EnhancedInputDeveloperSettings.h"
 
 #if WITH_EDITOR
 #include "Misc/DataValidation.h"
@@ -31,6 +32,11 @@ EDataValidationResult UInputMappingContext::IsDataValid(FDataValidationContext&
 }
 #endif	// WITH_EDITOR
 
+bool UInputMappingContext::ShouldShowInputModeQuery()
+{
+	return GetDefault<UEnhancedInputDeveloperSettings>()->bEnableInputModeFiltering;
+}
+
 void UInputMappingContext::PostLoad()
 {
 	Super::PostLoad();
diff --git a/Engine/Plugins/EnhancedInput/Source/EnhancedInput/Public/EnhancedInputDeveloperSettings.h b/Engine/Plugins/EnhancedInput/Source/EnhancedInput/Public/EnhancedInputDeveloperSettings.h
--- a/Engine/Plugins/EnhancedInput/Source/EnhancedInput/Public/EnhancedInputDeveloperSettings.h
+++ b/Engine/Plugins/EnhancedInput/Source/EnhancedInput/Public/EnhancedInputDeveloperSettings.h
@@ -112,6 +112,19 @@ public:
 	UPROPERTY(config, EditAnywhere, Category = "Enhanced Input", meta=(ConsoleVariable="EnhancedInput.OnlyTriggerLastActionInChord"))
 	uint8 bShouldOnlyTriggerLastActionInChord : 1;
 
+	/**
+	 * If true, then Enhanced Input will filter which Input Mapping Context mappings are processed
+	 * based on the current value of the "Current Input Mode". If a mapping context's query
+	 * does not match with the current input mode tag container, then its mappings will not be processed.
+	 *
+	 * Default Value is true.
+	 * 
+	 * @see IEnhancedInputSubsystemInterface::RebuildControlMappings
+	 * @see UEnhancedPlayerInput::CurrentInputMode
+	 */
+	UPROPERTY(config, EditAnywhere, Category = "Enhanced Input")
+	uint8 bEnableInputModeFiltering : 1;
+	
 	/**
 	 * If true, then a warning will be logged when a UPlayerMappableInputConfig that has been marked as deprecated is used.
 	 */
diff --git a/Engine/Plugins/EnhancedInput/Source/EnhancedInput/Public/EnhancedInputSubsystemInterface.h b/Engine/Plugins/EnhancedInput/Source/EnhancedInput/Public/EnhancedInputSubsystemInterface.h
--- a/Engine/Plugins/EnhancedInput/Source/EnhancedInput/Public/EnhancedInputSubsystemInterface.h
+++ b/Engine/Plugins/EnhancedInput/Source/EnhancedInput/Public/EnhancedInputSubsystemInterface.h
@@ -262,6 +262,59 @@ public:
 	UFUNCTION(BlueprintCallable, BlueprintCosmetic, Category = "Input", meta=(AutoCreateRefTerm = "Options"))
 	virtual void RemoveMappingContext(const UInputMappingContext* MappingContext, const FModifyContextOptions& Options = FModifyContextOptions());
 
+	/**
+	 * Returns the current input mode set on Enhanced Input. This mode will control what mappings are 
+	 * currently being processed by the player. If the Input Mapping Context's InputModeQuery is
+	 * not empty and does not match this tag container, then it's mappings will not be processed.
+	 */
+	UFUNCTION(BlueprintCallable, BlueprintCosmetic, Category = "Input")
+	virtual FGameplayTagContainer GetInputMode() const;
+
+	/**
+	 * Sets the current input mode on the player and triggers a rebuild of control mappings.
+	 * 
+	 * @param NewMode The new input mode tag container to use
+	 * @param Options Options for how to treat the control mappings rebuild.
+	 */
+	UFUNCTION(BlueprintCallable, BlueprintCosmetic, Category = "Input", meta=(AutoCreateRefTerm = "Options", AdvancedDisplay = "Options"))
+	virtual void SetInputMode(const FGameplayTagContainer& NewMode, const FModifyContextOptions& Options = FModifyContextOptions());
+
+	/**
+	 * Appends the given "TagsToAdd" to the current input mode tag container and triggers a rebuild of control mappings.
+	 *  
+	 * @param TagsToAdd	The gameplay tag container to append 
+	 * @param Options	Options for how to treat the control mappings rebuild
+	 */
+	UFUNCTION(BlueprintCallable, BlueprintCosmetic, Category = "Input", meta=(AutoCreateRefTerm = "Options", AdvancedDisplay = "Options"))
+	virtual void AppendTagsToInputMode(const FGameplayTagContainer& TagsToAdd, const FModifyContextOptions& Options = FModifyContextOptions());
+
+	/**
+	 * Adds the given "TagToAdd" to the current input mode tag container and triggers a rebuild of control mappings.
+	 * 
+	 * @param TagToAdd	The gameplay tag to add to the current input mode tag container
+	 * @param Options	Options for how to treat the control mappings rebuild
+	 */
+	UFUNCTION(BlueprintCallable, BlueprintCosmetic, Category = "Input", meta=(AutoCreateRefTerm = "Options", AdvancedDisplay = "Options"))
+	virtual void AddTagToInputMode(const FGameplayTag& TagToAdd, const FModifyContextOptions& Options = FModifyContextOptions());
+	
+	/**
+	 * Removes tags from the current input mode and triggers a rebuild of control mappings.
+	 * 
+	 * @param TagsToRemove Tags to remove from the current input mode tag container
+	 * @param Options		Options for how to treat the control mappings rebuild
+	 */
+	UFUNCTION(BlueprintCallable, BlueprintCosmetic, Category = "Input", meta=(AutoCreateRefTerm = "Options", AdvancedDisplay = "Options"))
+	virtual void RemoveTagsFromInputMode(const FGameplayTagContainer& TagsToRemove, const FModifyContextOptions& Options = FModifyContextOptions());
+
+	/**
+	 * Removes the given tag from the current input mode and triggers a rebuild of control mappings.
+	 * 
+	 * @param TagToRemove	The tag to remove from the current input mode 
+	 * @param Options		Options for how to treat the control mappings rebuild 
+	 */
+	UFUNCTION(BlueprintCallable, BlueprintCosmetic, Category = "Input", meta=(AutoCreateRefTerm = "Options", AdvancedDisplay = "Options"))
+	virtual void RemoveTagFromInputMode(const FGameplayTag& TagToRemove, const FModifyContextOptions& Options = FModifyContextOptions());
+	
 	/**
 	* Flag player for reapplication of all mapping contexts at the end of this frame.
 	* This is called automatically when adding or removing mappings contexts.
diff --git a/Engine/Plugins/EnhancedInput/Source/EnhancedInput/Public/EnhancedPlayerInput.h b/Engine/Plugins/EnhancedInput/Source/EnhancedInput/Public/EnhancedPlayerInput.h
--- a/Engine/Plugins/EnhancedInput/Source/EnhancedInput/Public/EnhancedPlayerInput.h
+++ b/Engine/Plugins/EnhancedInput/Source/EnhancedInput/Public/EnhancedPlayerInput.h
@@ -4,6 +4,7 @@
 
 #include "GameFramework/PlayerInput.h"
 #include "InputAction.h"
+#include "GameplayTagContainer.h"
 
 #include "EnhancedPlayerInput.generated.h"
 
@@ -98,9 +99,28 @@ public:
 
 	/** Returns the Time Dilation value that is currently effecting this input. */
 	float GetEffectiveTimeDilation() const;
+
+	/**
+	 * Returns a const ref to the current input mode. 
+	 */
+	const FGameplayTagContainer& GetCurrentInputMode() const;
 	
 protected:
 
+	/**
+	 * Sets the current input mode to be the given NewMode.
+	 * 
+	 * This should only be called by the Enhanced Input subsystem interface, and followed up by a call to RequestRebuildControlMappings.
+	 */
+	void SetCurrentInputMode(const FGameplayTagContainer& NewMode);
+
+	/**
+	 * Returns a mutable reference to the current input mode
+	 *
+	 * This should only be called by the Enhanced Input subsystem interface, and followed up by a call to RequestRebuildControlMappings.
+	 */
+	FGameplayTagContainer& GetCurrentInputMode();
+
 	virtual void EvaluateKeyMapState(const float DeltaTime, const bool bGamePaused, OUT TArray<TPair<FKey, FKeyState*>>& KeysWithEvents) override;
 	virtual void EvaluateInputDelegates(const TArray<UInputComponent*>& InputComponentStack, const float DeltaTime, const bool bGamePaused, const TArray<TPair<FKey, FKeyState*>>& KeysWithEvents) override;
 	
@@ -154,8 +174,12 @@ private:
 	template<typename T>
 	void GatherActionEventDataForActionMap(const T& ActionMap, TMap<TObjectPtr<const UInputAction>, FInputActionInstance>& FoundActionEventData) const;
 
-	/** Currently applied key mappings
+	/**
+	 * Currently applied key mappings
 	 * Note: Source reference only. Use EnhancedActionMappings for the actual mappings (with properly instanced triggers/modifiers)
+	 *
+	 * These mapping contexts will only have their mappings processed if the current input mode
+	 * matches the query set on them.
 	 */
 	UPROPERTY(Transient)
 	TMap<TObjectPtr<const UInputMappingContext>, int32> AppliedInputContexts;
@@ -164,6 +188,13 @@ private:
 	UPROPERTY(Transient)
 	TArray<FEnhancedActionKeyMapping> EnhancedActionMappings;
 
+	/**
+	 * The current input mode that is active on this player. If Input Mapping contexts have requirements
+	 * which this container does not meet, then their mappings will not be applied.
+	 */
+	UPROPERTY(Transient)
+	FGameplayTagContainer CurrentInputMode;
+
 	// Number of active binds by key
 	TMap<FKey, int32> EnhancedKeyBinds;
 
diff --git a/Engine/Plugins/EnhancedInput/Source/EnhancedInput/Public/InputMappingContext.h b/Engine/Plugins/EnhancedInput/Source/EnhancedInput/Public/InputMappingContext.h
--- a/Engine/Plugins/EnhancedInput/Source/EnhancedInput/Public/InputMappingContext.h
+++ b/Engine/Plugins/EnhancedInput/Source/EnhancedInput/Public/InputMappingContext.h
@@ -3,6 +3,7 @@
 #pragma once
 
 #include "EnhancedActionKeyMapping.h"
+#include "GameplayTagContainer.h"
 
 #include "InputMappingContext.generated.h"
 
@@ -34,6 +35,26 @@ public:
 	virtual EDataValidationResult IsDataValid(class FDataValidationContext& Context) const override;
 #endif
 
+	/**
+	 * If true, then the InputModeQuery must match the current input mode in order for these mappings to be processed.
+	 * If false, then the mappings on this context will not be filtered by the input mode.
+	 * 
+	 * @Note: bEnableInputModeFiltering must be enabled in the UEnhancedInputDeveloperSettings for this to be considered.
+	 */
+	UPROPERTY(BlueprintReadOnly, EditAnywhere, Category = "Input Modes", meta = (EditCondition="ShouldShowInputModeQuery()"))
+	bool bFilterMappingsByQuery = false;
+	
+	/**
+	 * Tag Query which will be matched against the current Enhanced Input Subsystem's input mode. 
+	 *
+	 * If this tag query does not match with the current input mode tag container, then the mappings
+	 * will not be processed.
+	 *
+	 * @Note: bEnableInputModeFiltering must be enabled in the UEnhancedInputDeveloperSettings for this to be considered.
+	 */
+	UPROPERTY(BlueprintReadOnly, EditAnywhere, Category = "Input Modes", meta = (EditCondition="bFilterMappingsByQuery && ShouldShowInputModeQuery()"))
+	FGameplayTagQuery InputModeQuery;
+
 	// Localized context descriptor
 	UPROPERTY(BlueprintReadOnly, EditAnywhere, Category = "Description", DisplayName = "Description")
 	FText ContextDescription;
@@ -43,6 +64,8 @@ public:
 
 public:
 
+	UFUNCTION()
+	static bool ShouldShowInputModeQuery();
 	
 	virtual void PostLoad() override;
 
