Deprecate TObjectPtr AssociatedInputAction to avoid GC issues when a GameFeaturePlugin unregisters

diff --git a/Engine/Plugins/EnhancedInput/Config/BaseEnhancedInput.ini b/Engine/Plugins/EnhancedInput/Config/BaseEnhancedInput.ini
--- a/Engine/Plugins/EnhancedInput/Config/BaseEnhancedInput.ini
+++ b/Engine/Plugins/EnhancedInput/Config/BaseEnhancedInput.ini
@@ -4,4 +4,6 @@
 ; 5.2
 +FunctionRedirects=(OldName="/Script/EnhancedInput.EnhancedInputSubsystemInterface.AddPlayerMappedKey", NewName="/Script/EnhancedInput.EnhancedInputSubsystemInterface.K2_AddPlayerMappedKeyInSlot")
 +FunctionRedirects=(OldName="/Script/EnhancedInput.EnhancedInputSubsystemInterface.RemovePlayerMappedKey", NewName="/Script/EnhancedInput.EnhancedInputSubsystemInterface.K2_RemovePlayerMappedKeyInSlot")
-+FunctionRedirects=(OldName="/Script/EnhancedInput.EnhancedInputSubsystemInterface.GetPlayerMappedKey", NewName="/Script/EnhancedInput.EnhancedInputSubsystemInterface.K2_GetPlayerMappedKeyInSlot")
\ No newline at end of file
++FunctionRedirects=(OldName="/Script/EnhancedInput.EnhancedInputSubsystemInterface.GetPlayerMappedKey", NewName="/Script/EnhancedInput.EnhancedInputSubsystemInterface.K2_GetPlayerMappedKeyInSlot")
+; 5.6
++PropertyRedirects=(OldName="EnhancedInput.PlayerKeyMapping.AssociatedInputAction",NewName="EnhancedInput.PlayerKeyMapping.AssociatedInputActionSoft")
\ No newline at end of file
diff --git a/Engine/Plugins/EnhancedInput/Source/EnhancedInput/Private/UserSettings/EnhancedInputUserSettings.cpp b/Engine/Plugins/EnhancedInput/Source/EnhancedInput/Private/UserSettings/EnhancedInputUserSettings.cpp
--- a/Engine/Plugins/EnhancedInput/Source/EnhancedInput/Private/UserSettings/EnhancedInputUserSettings.cpp
+++ b/Engine/Plugins/EnhancedInput/Source/EnhancedInput/Private/UserSettings/EnhancedInputUserSettings.cpp
@@ -104,7 +104,10 @@ FPlayerKeyMapping::FPlayerKeyMapping()
 	, DefaultKey(EKeys::Invalid)
 	, CurrentKey(EKeys::Invalid)
 	, HardwareDeviceId(FHardwareDeviceIdentifier::Invalid)
+PRAGMA_DISABLE_DEPRECATION_WARNINGS
 	, AssociatedInputAction(nullptr)
+PRAGMA_ENABLE_DEPRECATION_WARNINGS
+	, AssociatedInputActionSoft(nullptr)
 {
 }
 
@@ -120,7 +123,10 @@ FPlayerKeyMapping::FPlayerKeyMapping(
 	, DefaultKey(OriginalMapping.Key)
 	, CurrentKey(OriginalMapping.Key)
 	, HardwareDeviceId(InHardwareDevice)
-	, AssociatedInputAction(OriginalMapping.Action)
+PRAGMA_DISABLE_DEPRECATION_WARNINGS
+	, AssociatedInputAction(nullptr)
+PRAGMA_ENABLE_DEPRECATION_WARNINGS
+	, AssociatedInputActionSoft(OriginalMapping.Action)
 {
 }
 
@@ -159,7 +165,7 @@ FString FPlayerKeyMapping::ToString() const
 			*DefaultKey.ToString(),
 			*CurrentKey.ToString(),
 			*HardwareDeviceId.ToString(),
-			AssociatedInputAction ? *AssociatedInputAction->GetFName().ToString() : TEXT("NULL INPUT ACTION"));
+			AssociatedInputActionSoft.IsValid() ? *AssociatedInputActionSoft.Get()->GetFName().ToString() : TEXT("NULL INPUT ACTION"));
 }
 
 const FName FPlayerKeyMapping::GetMappingName() const
@@ -199,7 +205,7 @@ EHardwareDeviceSupportedFeatures::Type FPlayerKeyMapping::GetHardwareDeviceSuppo
 
 const UInputAction* FPlayerKeyMapping::GetAssociatedInputAction() const
 {
-	return AssociatedInputAction;
+	return AssociatedInputActionSoft.Get();
 }
 
 uint32 GetTypeHash(const FPlayerKeyMapping& InMapping)
@@ -245,7 +251,7 @@ void FPlayerKeyMapping::UpdateMetadataFromActionKeyMapping(const FEnhancedAction
 	
 	DisplayCategory = OriginalMapping.GetDisplayCategory();
 	DisplayName = OriginalMapping.GetDisplayName();
-	AssociatedInputAction = OriginalMapping.Action;
+	AssociatedInputActionSoft = OriginalMapping.Action;
 }
 
 void FPlayerKeyMapping::UpdateDefaultKeyFromActionKeyMapping(const FEnhancedActionKeyMapping& OriginalMapping)
@@ -956,7 +962,7 @@ void UEnhancedInputUserSettings::MapPlayerKey(const FMapPlayerKeyArgs& InArgs, F
 			PlayerMappingData.SetCurrentKey(InArgs.NewKey);
 			PlayerMappingData.DisplayName = ExistingMapping->DisplayName;
 			PlayerMappingData.DisplayCategory = ExistingMapping->DisplayCategory;
-			PlayerMappingData.AssociatedInputAction = ExistingMapping->AssociatedInputAction;
+			PlayerMappingData.AssociatedInputActionSoft = ExistingMapping->AssociatedInputActionSoft;
 			
 			const FSetElementId SetElem = MappingRow->Mappings.Add(PlayerMappingData);
 			OnKeyMappingUpdated(&MappingRow->Mappings.Get(SetElem), InArgs, false);
@@ -1414,28 +1420,6 @@ bool UEnhancedInputUserSettings::RegisterKeyMappingsToProfile(UEnhancedPlayerMap
 	return true;
 }
 
-bool UEnhancedInputUserSettings::UnregisterKeyMappingsToProfile(UEnhancedPlayerMappableKeyProfile& Profile, const UInputMappingContext* IMC)
-{
-	if (!IMC)
-	{
-		UE_LOG(LogEnhancedInput, Error, TEXT("Attempting to unregister a null mapping context with the user settings!"));
-		return false;
-	}
-
-	for (const FEnhancedActionKeyMapping& KeyMapping : IMC->GetMappings())
-	{
-		// Skip over non-player mappable keys
-		if (!KeyMapping.IsPlayerMappable())
-		{
-			continue;
-		}
-
-		const FName MappingName = KeyMapping.GetMappingName();
-		Profile.PlayerMappedKeys.Remove(MappingName);
-	}
-	return true;
-}
-
 void UEnhancedInputUserSettings::OnKeyMappingRegistered(FPlayerKeyMapping& RegisteredMapping, const FEnhancedActionKeyMapping& SourceMapping)
 {
 	// Does nothing by default
@@ -1448,13 +1432,7 @@ FHardwareDeviceIdentifier UEnhancedInputUserSettings::DetermineHardwareDeviceFor
 
 bool UEnhancedInputUserSettings::UnregisterInputMappingContext(const UInputMappingContext* IMC)
 {
-	bool bResult = true;
-	for (TPair<FGameplayTag, TObjectPtr<UEnhancedPlayerMappableKeyProfile>> Pair : SavedKeyProfiles)
-	{
-		bResult &= UnregisterKeyMappingsToProfile(*Pair.Value, IMC);
-	}
-	bResult &= RegisteredMappingContexts.Remove(IMC) != INDEX_NONE;
-	return bResult;
+	return RegisteredMappingContexts.Remove(IMC) != INDEX_NONE;
 }
 
 bool UEnhancedInputUserSettings::UnregisterInputMappingContexts(const TSet<UInputMappingContext*>& MappingContexts)
diff --git a/Engine/Plugins/EnhancedInput/Source/EnhancedInput/Public/UserSettings/EnhancedInputUserSettings.h b/Engine/Plugins/EnhancedInput/Source/EnhancedInput/Public/UserSettings/EnhancedInputUserSettings.h
--- a/Engine/Plugins/EnhancedInput/Source/EnhancedInput/Public/UserSettings/EnhancedInputUserSettings.h
+++ b/Engine/Plugins/EnhancedInput/Source/EnhancedInput/Public/UserSettings/EnhancedInputUserSettings.h
@@ -213,9 +213,13 @@ protected:
 	UPROPERTY(VisibleAnywhere, BlueprintReadOnly, Category="Enhanced Input|User Settings")
 	FHardwareDeviceIdentifier HardwareDeviceId;
 
+	UE_DEPRECATED(5.6, "This property is deprecated, use AssociatedInputActionSoft instead.")
+	UPROPERTY(Transient, VisibleAnywhere, BlueprintReadOnly, Category="Enhanced Input|User Settings", meta=(DeprecatedProperty, DeprecationMessage = "This property is deprecated, use AssociatedInputActionSoft instead."))
+	TObjectPtr<const UInputAction> AssociatedInputAction; 
+
 	/** The input action associated with this player key mapping */
 	UPROPERTY(Transient, VisibleAnywhere, BlueprintReadOnly, Category="Enhanced Input|User Settings")
-	TObjectPtr<const UInputAction> AssociatedInputAction; 
+	TSoftObjectPtr<const UInputAction> AssociatedInputActionSoft;
 };
 
 /**
@@ -697,8 +701,6 @@ protected:
 	 */
 	virtual bool RegisterKeyMappingsToProfile(UEnhancedPlayerMappableKeyProfile& Profile, const UInputMappingContext* IMC);
 
-	virtual bool UnregisterKeyMappingsToProfile(UEnhancedPlayerMappableKeyProfile& Profile, const UInputMappingContext* IMC);
-
 	/**
 	 * Determines the hardware device that an action key mapping is associated with. By default, the hardware will be FHardwareDeviceIdentifier::Invalid.
 	 *
