Undo CL: 37414823

Deprecate TObjectPtr AssociatedInputAction to avoid GC issues when a GameFeaturePlugin unregisters

diff --git a/Engine/Plugins/EnhancedInput/Config/BaseEnhancedInput.ini b/Engine/Plugins/EnhancedInput/Config/BaseEnhancedInput.ini
--- a/Engine/Plugins/EnhancedInput/Config/BaseEnhancedInput.ini
+++ b/Engine/Plugins/EnhancedInput/Config/BaseEnhancedInput.ini
@@ -4,6 +4,4 @@
 ; 5.2
 +FunctionRedirects=(OldName="/Script/EnhancedInput.EnhancedInputSubsystemInterface.AddPlayerMappedKey", NewName="/Script/EnhancedInput.EnhancedInputSubsystemInterface.K2_AddPlayerMappedKeyInSlot")
 +FunctionRedirects=(OldName="/Script/EnhancedInput.EnhancedInputSubsystemInterface.RemovePlayerMappedKey", NewName="/Script/EnhancedInput.EnhancedInputSubsystemInterface.K2_RemovePlayerMappedKeyInSlot")
-+FunctionRedirects=(OldName="/Script/EnhancedInput.EnhancedInputSubsystemInterface.GetPlayerMappedKey", NewName="/Script/EnhancedInput.EnhancedInputSubsystemInterface.K2_GetPlayerMappedKeyInSlot")
-; 5.6
-+PropertyRedirects=(OldName="EnhancedInput.PlayerKeyMapping.AssociatedInputAction",NewName="EnhancedInput.PlayerKeyMapping.AssociatedInputActionSoft")
\ No newline at end of file
++FunctionRedirects=(OldName="/Script/EnhancedInput.EnhancedInputSubsystemInterface.GetPlayerMappedKey", NewName="/Script/EnhancedInput.EnhancedInputSubsystemInterface.K2_GetPlayerMappedKeyInSlot")
\ No newline at end of file
diff --git a/Engine/Plugins/EnhancedInput/Source/EnhancedInput/Private/UserSettings/EnhancedInputUserSettings.cpp b/Engine/Plugins/EnhancedInput/Source/EnhancedInput/Private/UserSettings/EnhancedInputUserSettings.cpp
--- a/Engine/Plugins/EnhancedInput/Source/EnhancedInput/Private/UserSettings/EnhancedInputUserSettings.cpp
+++ b/Engine/Plugins/EnhancedInput/Source/EnhancedInput/Private/UserSettings/EnhancedInputUserSettings.cpp
@@ -104,10 +104,7 @@ FPlayerKeyMapping::FPlayerKeyMapping()
 	, DefaultKey(EKeys::Invalid)
 	, CurrentKey(EKeys::Invalid)
 	, HardwareDeviceId(FHardwareDeviceIdentifier::Invalid)
-PRAGMA_DISABLE_DEPRECATION_WARNINGS
 	, AssociatedInputAction(nullptr)
-PRAGMA_ENABLE_DEPRECATION_WARNINGS
-	, AssociatedInputActionSoft(nullptr)
 {
 }
 
@@ -123,10 +120,7 @@ FPlayerKeyMapping::FPlayerKeyMapping(
 	, DefaultKey(OriginalMapping.Key)
 	, CurrentKey(OriginalMapping.Key)
 	, HardwareDeviceId(InHardwareDevice)
-PRAGMA_DISABLE_DEPRECATION_WARNINGS
-	, AssociatedInputAction(nullptr)
-PRAGMA_ENABLE_DEPRECATION_WARNINGS
-	, AssociatedInputActionSoft(OriginalMapping.Action)
+	, AssociatedInputAction(OriginalMapping.Action)
 {
 }
 
@@ -165,7 +159,7 @@ FString FPlayerKeyMapping::ToString() const
 			*DefaultKey.ToString(),
 			*CurrentKey.ToString(),
 			*HardwareDeviceId.ToString(),
-			AssociatedInputActionSoft.IsValid() ? *AssociatedInputActionSoft.Get()->GetFName().ToString() : TEXT("NULL INPUT ACTION"));
+			AssociatedInputAction ? *AssociatedInputAction->GetFName().ToString() : TEXT("NULL INPUT ACTION"));
 }
 
 const FName FPlayerKeyMapping::GetMappingName() const
@@ -205,7 +199,7 @@ EHardwareDeviceSupportedFeatures::Type FPlayerKeyMapping::GetHardwareDeviceSuppo
 
 const UInputAction* FPlayerKeyMapping::GetAssociatedInputAction() const
 {
-	return AssociatedInputActionSoft.Get();
+	return AssociatedInputAction;
 }
 
 uint32 GetTypeHash(const FPlayerKeyMapping& InMapping)
@@ -251,7 +245,7 @@ void FPlayerKeyMapping::UpdateMetadataFromActionKeyMapping(const FEnhancedAction
 	
 	DisplayCategory = OriginalMapping.GetDisplayCategory();
 	DisplayName = OriginalMapping.GetDisplayName();
-	AssociatedInputActionSoft = OriginalMapping.Action;
+	AssociatedInputAction = OriginalMapping.Action;
 }
 
 void FPlayerKeyMapping::UpdateDefaultKeyFromActionKeyMapping(const FEnhancedActionKeyMapping& OriginalMapping)
@@ -962,7 +956,7 @@ void UEnhancedInputUserSettings::MapPlayerKey(const FMapPlayerKeyArgs& InArgs, F
 			PlayerMappingData.SetCurrentKey(InArgs.NewKey);
 			PlayerMappingData.DisplayName = ExistingMapping->DisplayName;
 			PlayerMappingData.DisplayCategory = ExistingMapping->DisplayCategory;
-			PlayerMappingData.AssociatedInputActionSoft = ExistingMapping->AssociatedInputActionSoft;
+			PlayerMappingData.AssociatedInputAction = ExistingMapping->AssociatedInputAction;
 			
 			const FSetElementId SetElem = MappingRow->Mappings.Add(PlayerMappingData);
 			OnKeyMappingUpdated(&MappingRow->Mappings.Get(SetElem), InArgs, false);
@@ -1420,6 +1414,28 @@ bool UEnhancedInputUserSettings::RegisterKeyMappingsToProfile(UEnhancedPlayerMap
 	return true;
 }
 
+bool UEnhancedInputUserSettings::UnregisterKeyMappingsToProfile(UEnhancedPlayerMappableKeyProfile& Profile, const UInputMappingContext* IMC)
+{
+	if (!IMC)
+	{
+		UE_LOG(LogEnhancedInput, Error, TEXT("Attempting to unregister a null mapping context with the user settings!"));
+		return false;
+	}
+
+	for (const FEnhancedActionKeyMapping& KeyMapping : IMC->GetMappings())
+	{
+		// Skip over non-player mappable keys
+		if (!KeyMapping.IsPlayerMappable())
+		{
+			continue;
+		}
+
+		const FName MappingName = KeyMapping.GetMappingName();
+		Profile.PlayerMappedKeys.Remove(MappingName);
+	}
+	return true;
+}
+
 void UEnhancedInputUserSettings::OnKeyMappingRegistered(FPlayerKeyMapping& RegisteredMapping, const FEnhancedActionKeyMapping& SourceMapping)
 {
 	// Does nothing by default
@@ -1432,7 +1448,13 @@ FHardwareDeviceIdentifier UEnhancedInputUserSettings::DetermineHardwareDeviceFor
 
 bool UEnhancedInputUserSettings::UnregisterInputMappingContext(const UInputMappingContext* IMC)
 {
-	return RegisteredMappingContexts.Remove(IMC) != INDEX_NONE;
+	bool bResult = true;
+	for (TPair<FGameplayTag, TObjectPtr<UEnhancedPlayerMappableKeyProfile>> Pair : SavedKeyProfiles)
+	{
+		bResult &= UnregisterKeyMappingsToProfile(*Pair.Value, IMC);
+	}
+	bResult &= RegisteredMappingContexts.Remove(IMC) != INDEX_NONE;
+	return bResult;
 }
 
 bool UEnhancedInputUserSettings::UnregisterInputMappingContexts(const TSet<UInputMappingContext*>& MappingContexts)
diff --git a/Engine/Plugins/EnhancedInput/Source/EnhancedInput/Public/UserSettings/EnhancedInputUserSettings.h b/Engine/Plugins/EnhancedInput/Source/EnhancedInput/Public/UserSettings/EnhancedInputUserSettings.h
--- a/Engine/Plugins/EnhancedInput/Source/EnhancedInput/Public/UserSettings/EnhancedInputUserSettings.h
+++ b/Engine/Plugins/EnhancedInput/Source/EnhancedInput/Public/UserSettings/EnhancedInputUserSettings.h
@@ -213,13 +213,9 @@ protected:
 	UPROPERTY(VisibleAnywhere, BlueprintReadOnly, Category="Enhanced Input|User Settings")
 	FHardwareDeviceIdentifier HardwareDeviceId;
 
-	UE_DEPRECATED(5.6, "This property is deprecated, use AssociatedInputActionSoft instead.")
-	UPROPERTY(Transient, VisibleAnywhere, BlueprintReadOnly, Category="Enhanced Input|User Settings", meta=(DeprecatedProperty, DeprecationMessage = "This property is deprecated, use AssociatedInputActionSoft instead."))
-	TObjectPtr<const UInputAction> AssociatedInputAction; 
-
 	/** The input action associated with this player key mapping */
 	UPROPERTY(Transient, VisibleAnywhere, BlueprintReadOnly, Category="Enhanced Input|User Settings")
-	TSoftObjectPtr<const UInputAction> AssociatedInputActionSoft;
+	TObjectPtr<const UInputAction> AssociatedInputAction; 
 };
 
 /**
@@ -701,6 +697,8 @@ protected:
 	 */
 	virtual bool RegisterKeyMappingsToProfile(UEnhancedPlayerMappableKeyProfile& Profile, const UInputMappingContext* IMC);
 
+	virtual bool UnregisterKeyMappingsToProfile(UEnhancedPlayerMappableKeyProfile& Profile, const UInputMappingContext* IMC);
+
 	/**
 	 * Determines the hardware device that an action key mapping is associated with. By default, the hardware will be FHardwareDeviceIdentifier::Invalid.
 	 *
