Remove unused enhanced input setting bLogOnDeprecatedConfigUsed

#rnx

diff --git a/Engine/Plugins/EnhancedInput/Source/EnhancedInput/Private/EnhancedInputDeveloperSettings.cpp b/Engine/Plugins/EnhancedInput/Source/EnhancedInput/Private/EnhancedInputDeveloperSettings.cpp
--- a/Engine/Plugins/EnhancedInput/Source/EnhancedInput/Private/EnhancedInputDeveloperSettings.cpp
+++ b/Engine/Plugins/EnhancedInput/Source/EnhancedInput/Private/EnhancedInputDeveloperSettings.cpp
@@ -31,7 +31,6 @@ UEnhancedInputDeveloperSettings::UEnhancedInputDeveloperSettings(const FObjectIn
 	, bEnableDefaultMappingContexts(true)
 	, bShouldOnlyTriggerLastActionInChord(true)
 	, bEnableInputModeFiltering(true)
-	, bLogOnDeprecatedConfigUsed(true)
 	, bEnableWorldSubsystem(false)
 	, bShouldLogAllWorldSubsystemInputs(false)
 {
diff --git a/Engine/Plugins/EnhancedInput/Source/EnhancedInput/Public/EnhancedInputDeveloperSettings.h b/Engine/Plugins/EnhancedInput/Source/EnhancedInput/Public/EnhancedInputDeveloperSettings.h
--- a/Engine/Plugins/EnhancedInput/Source/EnhancedInput/Public/EnhancedInputDeveloperSettings.h
+++ b/Engine/Plugins/EnhancedInput/Source/EnhancedInput/Public/EnhancedInputDeveloperSettings.h
@@ -125,12 +125,6 @@ public:
 	UPROPERTY(config, EditAnywhere, Category = "Enhanced Input")
 	uint8 bEnableInputModeFiltering : 1;
 	
-	/**
-	 * If true, then a warning will be logged when a UPlayerMappableInputConfig that has been marked as deprecated is used.
-	 */
-	UPROPERTY(config, EditAnywhere, Category = "Enhanced Input")
-	uint8 bLogOnDeprecatedConfigUsed : 1;
-	
 	/** If true, then the world subsystem will be created. */
 	UPROPERTY(config, EditAnywhere, Category = "Enhanced Input|World Subsystem", meta=(DisplayName="Enable World Subsystem (Experimental)", DisplayPriority = 1))
 	uint8 bEnableWorldSubsystem : 1;
