Clean up some misc. Enhanced Input TODO's and tooltips on BP nodes.

#rnx

diff --git a/Engine/Plugins/EnhancedInput/Source/EnhancedInput/Public/EnhancedInputComponent.h b/Engine/Plugins/EnhancedInput/Source/EnhancedInput/Public/EnhancedInputComponent.h
--- a/Engine/Plugins/EnhancedInput/Source/EnhancedInput/Public/EnhancedInputComponent.h
+++ b/Engine/Plugins/EnhancedInput/Source/EnhancedInput/Public/EnhancedInputComponent.h
@@ -20,6 +20,12 @@ DECLARE_DYNAMIC_DELEGATE_TwoParams(FInputDebugKeyHandlerDynamicSignature, FKey,
 DECLARE_DELEGATE(FEnhancedInputActionHandlerSignature);
 DECLARE_DELEGATE_OneParam(FEnhancedInputActionHandlerValueSignature, const FInputActionValue&);
 DECLARE_DELEGATE_OneParam(FEnhancedInputActionHandlerInstanceSignature, const FInputActionInstance&);	// Provides full access to value and timers
+/**
+ * @param ActionValue		The combined value of all inputs mapped to this action.
+ * @param ElapsedTime		The amount of time in seconds that has elapsed since this input action has begun processing its input values. 
+ * @param TriggeredTime		The amount of time in seconds that this input action has been in the "triggered" state.
+ * @param SourceAction		The input action which caused this event to fire. 
+ */
 DECLARE_DYNAMIC_DELEGATE_FourParams(FEnhancedInputActionHandlerDynamicSignature, FInputActionValue, ActionValue, float, ElapsedTime, float, TriggeredTime, const UInputAction*, SourceAction);
 
 // By default this behavior is on in the editor only. If you would like to turn this
diff --git a/Engine/Plugins/EnhancedInput/Source/InputBlueprintNodes/Private/K2Node_EnhancedInputAction.cpp b/Engine/Plugins/EnhancedInput/Source/InputBlueprintNodes/Private/K2Node_EnhancedInputAction.cpp
--- a/Engine/Plugins/EnhancedInput/Source/InputBlueprintNodes/Private/K2Node_EnhancedInputAction.cpp
+++ b/Engine/Plugins/EnhancedInput/Source/InputBlueprintNodes/Private/K2Node_EnhancedInputAction.cpp
@@ -93,11 +93,16 @@ void UK2Node_EnhancedInputAction::AllocateDefaultPins()
 	AdvancedPinDisplay = ENodeAdvancedPins::Hidden;
 
 	UEdGraphPin* ValuePin = CreatePin(EGPD_Output, UK2Node_GetInputActionValue::GetValueCategory(InputAction), UK2Node_GetInputActionValue::GetValueSubCategory(InputAction), UK2Node_GetInputActionValue::GetValueSubCategoryObject(InputAction), ActionValuePinName);
-	
+	Schema->ConstructBasicPinTooltip(*ValuePin, LOCTEXT("ValuePinDescription", "The combined value of all inputs mapped to this action."), ValuePin->PinToolTip);
 	Schema->SetPinAutogeneratedDefaultValueBasedOnType(ValuePin);
 
-	CreatePin(EGPD_Output, UEdGraphSchema_K2::PC_Real, UEdGraphSchema_K2::PC_Double, ElapsedSecondsPinName)->bAdvancedView = true;
-	CreatePin(EGPD_Output, UEdGraphSchema_K2::PC_Real, UEdGraphSchema_K2::PC_Double, TriggeredSecondsPinName)->bAdvancedView = true;
+	UEdGraphPin* ElapsedSecondsPin = CreatePin(EGPD_Output, UEdGraphSchema_K2::PC_Real, UEdGraphSchema_K2::PC_Double, ElapsedSecondsPinName);
+	ElapsedSecondsPin->bAdvancedView = true;
+	Schema->ConstructBasicPinTooltip(*ElapsedSecondsPin, LOCTEXT("ElapsedSecondsPinDescription", "The amount of time in seconds that has elapsed since this input action has begun processing its input values."), ElapsedSecondsPin->PinToolTip);
+	
+	UEdGraphPin* TriggeredSecondsPin = CreatePin(EGPD_Output, UEdGraphSchema_K2::PC_Real, UEdGraphSchema_K2::PC_Double, TriggeredSecondsPinName);
+	TriggeredSecondsPin->bAdvancedView = true;
+	Schema->ConstructBasicPinTooltip(*TriggeredSecondsPin, LOCTEXT("TriggeredSecondsPinDescription", "The amount of time in seconds that this input action has been in the \"triggered\" state."), TriggeredSecondsPin->PinToolTip);	
 
 	if(InputAction)
 	{
@@ -184,7 +189,6 @@ FName UK2Node_EnhancedInputAction::GetActionName() const
 
 FText UK2Node_EnhancedInputAction::GetNodeTitle(ENodeTitleType::Type TitleType) const
 {
-	// TODO: Is Using InputAction->GetFName okay here? Full Asset path would be better for disambiguation.
 	if (TitleType == ENodeTitleType::MenuTitle)
 	{
 		return FText::FromName(GetActionName());
@@ -472,7 +476,7 @@ FText UK2Node_EnhancedInputAction::GetMenuCategory() const
 	if (CachedCategory.IsOutOfDate(this))
 	{
 		// FText::Format() is slow, so we cache this to save on performance
-		CachedCategory.SetCachedText(FEditorCategoryUtils::BuildCategoryString(FCommonEditorCategory::Input, LOCTEXT("ActionMenuCategory", "Enhanced Action Events")), this);	// TODO: Rename Action Events once old action system is removed
+		CachedCategory.SetCachedText(FEditorCategoryUtils::BuildCategoryString(FCommonEditorCategory::Input, LOCTEXT("ActionMenuCategory", "Enhanced Input Actions")), this);
 	}
 	return CachedCategory;
 }
@@ -487,7 +491,6 @@ FBlueprintNodeSignature UK2Node_EnhancedInputAction::GetSignature() const
 
 TSharedPtr<FEdGraphSchemaAction> UK2Node_EnhancedInputAction::GetEventNodeAction(const FText& ActionCategory)
 {
-	// TODO: Custom EdGraphSchemaAction required?
 	TSharedPtr<FEdGraphSchemaAction_K2InputAction> EventNodeAction = MakeShareable(new FEdGraphSchemaAction_K2InputAction(ActionCategory, GetNodeTitle(ENodeTitleType::EditableTitle), GetTooltipText(), 0));
 	EventNodeAction->NodeTemplate = this;
 	return EventNodeAction;
diff --git a/Engine/Plugins/EnhancedInput/Source/InputBlueprintNodes/Private/K2Node_GetInputActionValue.cpp b/Engine/Plugins/EnhancedInput/Source/InputBlueprintNodes/Private/K2Node_GetInputActionValue.cpp
--- a/Engine/Plugins/EnhancedInput/Source/InputBlueprintNodes/Private/K2Node_GetInputActionValue.cpp
+++ b/Engine/Plugins/EnhancedInput/Source/InputBlueprintNodes/Private/K2Node_GetInputActionValue.cpp
@@ -142,10 +142,11 @@ void UK2Node_GetInputActionValue::ExpandNode(class FKismetCompilerContext& Compi
 
 	// And finally hook the InputActionValueAccessor's return value up to the GetInputActionValue return pin's connected nodes, automatically inserting appropriate conversion nodes.
 	TArray<UEdGraphPin*> LinkedTo(ReturnValue->LinkedTo);
+	UEdGraphPin* AccessorReturnPin = AccessorNode->FindPinChecked(UEdGraphSchema_K2::PN_ReturnValue);
+	
 	for (UEdGraphPin* EachLink : LinkedTo)
 	{
-		// TODO: Batch connections by conversion node type so conversion node is only hit once.
-		Schema->TryCreateConnection(AccessorNode->FindPinChecked(UEdGraphSchema_K2::PN_ReturnValue), EachLink);
+		Schema->TryCreateConnection(AccessorReturnPin, EachLink);
 	}
 }
 
diff --git a/Engine/Plugins/EnhancedInput/Source/InputBlueprintNodes/Private/K2Node_InputActionValueAccessor.cpp b/Engine/Plugins/EnhancedInput/Source/InputBlueprintNodes/Private/K2Node_InputActionValueAccessor.cpp
--- a/Engine/Plugins/EnhancedInput/Source/InputBlueprintNodes/Private/K2Node_InputActionValueAccessor.cpp
+++ b/Engine/Plugins/EnhancedInput/Source/InputBlueprintNodes/Private/K2Node_InputActionValueAccessor.cpp
@@ -41,15 +41,4 @@ void UK2Node_InputActionValueAccessor::RegisterDynamicBinding(UDynamicBlueprintB
 	Binding.FunctionNameToBind = GET_FUNCTION_NAME_CHECKED(UEnhancedInputLibrary, GetBoundActionValue);
 
 	InputActionBindingObject->InputActionValueBindings.Add(Binding);
-}
-// TODO: Do we care about unique signatures per input action?
-//
-//FBlueprintNodeSignature UK2Node_InputActionValueAccessor::GetSignature() const
-//{
-//	FBlueprintNodeSignature NodeSignature = Super::GetSignature();
-//	NodeSignature.AddKeyValue(InputAction->GetFName().ToString());
-//
-//	return NodeSignature;
-//}
-
-
+}
\ No newline at end of file
diff --git a/Engine/Plugins/EnhancedInput/Source/InputBlueprintNodes/Private/K2Node_InputDebugKey.cpp b/Engine/Plugins/EnhancedInput/Source/InputBlueprintNodes/Private/K2Node_InputDebugKey.cpp
--- a/Engine/Plugins/EnhancedInput/Source/InputBlueprintNodes/Private/K2Node_InputDebugKey.cpp
+++ b/Engine/Plugins/EnhancedInput/Source/InputBlueprintNodes/Private/K2Node_InputDebugKey.cpp
@@ -135,14 +135,6 @@ FText UK2Node_InputDebugKey::GetKeyText() const
 
 FText UK2Node_InputDebugKey::GetNodeTitle(ENodeTitleType::Type TitleType) const
 {
-	// TODO: The menu title version makes it very confusing to separate from the basic input system's key bindings.
-#if 0
-	if (TitleType == ENodeTitleType::MenuTitle)
-	{
-		return GetKeyText();
-	}
-#endif 
-
 	if (CachedNodeTitle.IsOutOfDate(this))
 	{
 		FFormatNamedArguments Args;
