Only spawn a new enhanced input action node if one does not exist in the graph already.

diff --git a/Engine/Plugins/EnhancedInput/Source/InputBlueprintNodes/Private/K2Node_EnhancedInputAction.cpp b/Engine/Plugins/EnhancedInput/Source/InputBlueprintNodes/Private/K2Node_EnhancedInputAction.cpp
--- a/Engine/Plugins/EnhancedInput/Source/InputBlueprintNodes/Private/K2Node_EnhancedInputAction.cpp
+++ b/Engine/Plugins/EnhancedInput/Source/InputBlueprintNodes/Private/K2Node_EnhancedInputAction.cpp
@@ -23,6 +23,7 @@
 #include "EnhancedInputEditorSettings.h"
 #include "WidgetBlueprint.h"
 #include "Blueprint/UserWidget.h"
+#include "BlueprintNodeTemplateCache.h"
 
 #include UE_INLINE_GENERATED_CPP_BY_NAME(K2Node_EnhancedInputAction)
 
@@ -60,6 +61,61 @@ void ForEachEventPinName(TFunctionRef<void(ETriggerEvent Event, FName PinName)>
 	}
 }
 
+UInputActionEventNodeSpawner* UInputActionEventNodeSpawner::Create(TSubclassOf<UEdGraphNode> const NodeClass, TObjectPtr<const UInputAction> InAction)
+{
+	check(NodeClass != nullptr);
+	check(NodeClass->IsChildOf<UEdGraphNode>());
+	check(InAction != nullptr);
+	
+	UInputActionEventNodeSpawner* NodeSpawner = NewObject<UInputActionEventNodeSpawner>(GetTransientPackage());
+	NodeSpawner->NodeClass = NodeClass;
+	NodeSpawner->WeakActionPtr = InAction;
+	
+	return NodeSpawner;
+}
+
+UEdGraphNode* UInputActionEventNodeSpawner::Invoke(UEdGraph* ParentGraph, FBindingSet const& Bindings, FVector2D const Location) const
+{
+	check(ParentGraph != nullptr);
+	UBlueprint* Blueprint = FBlueprintEditorUtils::FindBlueprintForGraphChecked(ParentGraph);
+	
+	if (!FBlueprintNodeTemplateCache::IsTemplateOuter(ParentGraph))
+	{
+		// Look to see if a node for this input action already exists. If it does, just return that
+		// which will jump the focus to it.
+		if (UK2Node* PreExistingNode = FindExistingNode(Blueprint))
+		{
+			return PreExistingNode;
+		}		
+	}
+	
+	return Super::Invoke(ParentGraph, Bindings, Location);
+}
+
+UK2Node* UInputActionEventNodeSpawner::FindExistingNode(const UBlueprint* Blueprint) const
+{
+	// We don't want references to node spawners to be keeping any input action assets from GC
+	// if you unload a plugin for example, so we keep it as a weak pointer.
+	TStrongObjectPtr<const UInputAction> ActionPtr = WeakActionPtr.Pin();
+	if (!ActionPtr)
+	{
+		return nullptr;
+	}
+	
+	TArray<UK2Node_EnhancedInputAction*> AllInputActionNodes;
+	FBlueprintEditorUtils::GetAllNodesOfClass<UK2Node_EnhancedInputAction>(Blueprint, AllInputActionNodes);
+
+	for (UK2Node_EnhancedInputAction* Node : AllInputActionNodes)
+	{
+		if (Node->InputAction == ActionPtr)
+		{
+			return Node;
+		}
+	}
+	
+	return nullptr;
+}
+
 void UK2Node_EnhancedInputAction::AllocateDefaultPins()
 {
 	PreloadObject((UObject*)InputAction);
@@ -442,13 +498,13 @@ void UK2Node_EnhancedInputAction::GetMenuActions(FBlueprintActionDatabaseRegistr
 		AssetRegistry.GetAssetsByClass(UInputAction::StaticClass()->GetClassPathName(), ActionAssets, true);
 		for (const FAssetData& ActionAsset : ActionAssets)
 		{
-			UBlueprintNodeSpawner* NodeSpawner = UBlueprintNodeSpawner::Create(GetClass());
-			check(NodeSpawner != nullptr);
-
 			if (FPackageName::GetPackageMountPoint(ActionAsset.PackageName.ToString()) != NAME_None)
 			{
 				if (const UInputAction* Action = Cast<const UInputAction>(ActionAsset.GetAsset()))
 				{
+					UBlueprintNodeSpawner* NodeSpawner = UInputActionEventNodeSpawner::Create(GetClass(), Action);
+					check(NodeSpawner != nullptr);
+					
 					NodeSpawner->CustomizeNodeDelegate = UBlueprintNodeSpawner::FCustomizeNodeDelegate::CreateStatic(CustomizeInputNodeLambda, TWeakObjectPtr<const UInputAction>(Action));
 					ActionRegistrar.AddBlueprintAction(Action, NodeSpawner);
 				}
diff --git a/Engine/Plugins/EnhancedInput/Source/InputBlueprintNodes/Private/K2Node_InputDebugKey.cpp b/Engine/Plugins/EnhancedInput/Source/InputBlueprintNodes/Private/K2Node_InputDebugKey.cpp
--- a/Engine/Plugins/EnhancedInput/Source/InputBlueprintNodes/Private/K2Node_InputDebugKey.cpp
+++ b/Engine/Plugins/EnhancedInput/Source/InputBlueprintNodes/Private/K2Node_InputDebugKey.cpp
@@ -16,6 +16,7 @@
 #include "GameFramework/InputSettings.h"
 #include "Editor.h"								// for FEditorDelegates::OnEnableGestureRecognizerChanged
 #include "Blueprint/UserWidget.h"
+#include "BlueprintNodeTemplateCache.h"
 
 #include UE_INLINE_GENERATED_CPP_BY_NAME(K2Node_InputDebugKey)
 
@@ -41,6 +42,52 @@ UK2Node_InputDebugKey::UK2Node_InputDebugKey(const FObjectInitializer& ObjectIni
 	SetEnabledState(ENodeEnabledState::DevelopmentOnly, false);
 }
 
+UInputDebugKeyEventNodeSpawner* UInputDebugKeyEventNodeSpawner::Create(TSubclassOf<UEdGraphNode> const NodeClass, const FKey& InDebugKey)
+{
+	check(NodeClass != nullptr);
+	check(NodeClass->IsChildOf<UEdGraphNode>());
+	
+	UInputDebugKeyEventNodeSpawner* NodeSpawner = NewObject<UInputDebugKeyEventNodeSpawner>(GetTransientPackage());
+	NodeSpawner->NodeClass = NodeClass;
+	NodeSpawner->DebugKey = InDebugKey;
+	
+	return NodeSpawner;
+}
+
+UEdGraphNode* UInputDebugKeyEventNodeSpawner::Invoke(UEdGraph* ParentGraph, FBindingSet const& Bindings, FVector2D const Location) const
+{
+	check(ParentGraph != nullptr);
+	UBlueprint* Blueprint = FBlueprintEditorUtils::FindBlueprintForGraphChecked(ParentGraph);
+	
+	if (!FBlueprintNodeTemplateCache::IsTemplateOuter(ParentGraph))
+	{
+		// Look to see if a node for this input action already exists. If it does, just return that
+		// which will jump the focus to it.
+		if (UK2Node* PreExistingNode = FindExistingNode(Blueprint))
+		{
+			return PreExistingNode;
+		}		
+	}
+	
+	return Super::Invoke(ParentGraph, Bindings, Location);
+}
+
+UK2Node* UInputDebugKeyEventNodeSpawner::FindExistingNode(const UBlueprint* Blueprint) const
+{
+	TArray<UK2Node_InputDebugKey*> AllDebugKeyNodes;
+	FBlueprintEditorUtils::GetAllNodesOfClass<UK2Node_InputDebugKey>(Blueprint, AllDebugKeyNodes);
+
+	for (UK2Node_InputDebugKey* Node : AllDebugKeyNodes)
+	{
+		if (Node->InputKey == DebugKey)
+		{
+			return Node;
+		}
+	}
+	
+	return nullptr;
+}
+
 void UK2Node_InputDebugKey::PostEditChangeProperty(FPropertyChangedEvent& PropertyChangedEvent)
 {
 	Super::PostEditChangeProperty(PropertyChangedEvent);
@@ -135,14 +182,6 @@ FText UK2Node_InputDebugKey::GetKeyText() const
 
 FText UK2Node_InputDebugKey::GetNodeTitle(ENodeTitleType::Type TitleType) const
 {
-	// TODO: The menu title version makes it very confusing to separate from the basic input system's key bindings.
-#if 0
-	if (TitleType == ENodeTitleType::MenuTitle)
-	{
-		return GetKeyText();
-	}
-#endif 
-
 	if (CachedNodeTitle.IsOutOfDate(this))
 	{
 		FFormatNamedArguments Args;
@@ -408,7 +447,7 @@ void UK2Node_InputDebugKey::GetMenuActions(FBlueprintActionDatabaseRegistrar& Ac
 				continue;
 			}
 
-			UBlueprintNodeSpawner* NodeSpawner = UBlueprintNodeSpawner::Create(GetClass());
+			UBlueprintNodeSpawner* NodeSpawner = UInputDebugKeyEventNodeSpawner::Create(GetClass(), Key);
 			check(NodeSpawner != nullptr);
 
 			NodeSpawner->CustomizeNodeDelegate = UBlueprintNodeSpawner::FCustomizeNodeDelegate::CreateStatic(CustomizeInputNodeLambda, Key);
diff --git a/Engine/Plugins/EnhancedInput/Source/InputBlueprintNodes/Private/K2Node_InputDebugKey.h b/Engine/Plugins/EnhancedInput/Source/InputBlueprintNodes/Private/K2Node_InputDebugKey.h
--- a/Engine/Plugins/EnhancedInput/Source/InputBlueprintNodes/Private/K2Node_InputDebugKey.h
+++ b/Engine/Plugins/EnhancedInput/Source/InputBlueprintNodes/Private/K2Node_InputDebugKey.h
@@ -2,6 +2,7 @@
 
 #pragma once
 
+#include "BlueprintNodeSpawner.h"
 #include "EdGraph/EdGraphNodeUtils.h"
 #include "K2Node.h"
 #include "K2Node_EventNodeInterface.h"
@@ -14,6 +15,26 @@ struct FBlueprintNodeSignature;
 class FBlueprintActionDatabaseRegistrar;
 class FKismetCompilerContext;
 
+/**
+ * A node spawner which will check if there is already adebug key event node (UInputDebugKeyEventNodeSpawner)
+ * before spawning a new one in the graph. This makes the behavior of debug key nodes the same as other
+ * event nodes. 
+ */
+UCLASS(Transient)
+class UInputDebugKeyEventNodeSpawner final : public UBlueprintNodeSpawner
+{
+	GENERATED_BODY()
+public:
+	static UInputDebugKeyEventNodeSpawner* Create(TSubclassOf<UEdGraphNode> const NodeClass, const FKey& DebugKey);
+private:
+	virtual UEdGraphNode* Invoke(UEdGraph* ParentGraph, FBindingSet const& Bindings, FVector2D const Location) const override;
+
+	UK2Node* FindExistingNode(const UBlueprint* Blueprint) const;
+
+	UPROPERTY()
+	FKey DebugKey = EKeys::Invalid;
+};
+
 UCLASS()
 class INPUTBLUEPRINTNODES_API UK2Node_InputDebugKey : public UK2Node, public IK2Node_EventNodeInterface
 {
diff --git a/Engine/Plugins/EnhancedInput/Source/InputBlueprintNodes/Public/K2Node_EnhancedInputAction.h b/Engine/Plugins/EnhancedInput/Source/InputBlueprintNodes/Public/K2Node_EnhancedInputAction.h
--- a/Engine/Plugins/EnhancedInput/Source/InputBlueprintNodes/Public/K2Node_EnhancedInputAction.h
+++ b/Engine/Plugins/EnhancedInput/Source/InputBlueprintNodes/Public/K2Node_EnhancedInputAction.h
@@ -5,6 +5,7 @@
 #include "EdGraph/EdGraphNodeUtils.h"
 #include "K2Node.h"
 #include "K2Node_EventNodeInterface.h"
+#include "BlueprintNodeSpawner.h"
 
 #include "K2Node_EnhancedInputAction.generated.h"
 
@@ -16,6 +17,28 @@ struct FBlueprintNodeSignature;
 class FBlueprintActionDatabaseRegistrar;
 class FKismetCompilerContext;
 
+/**
+ * A node spawner which will check if there is already an enhanced input event node (UK2Node_EnhancedInputAction)
+ * before spawning a new one in the graph. This makes the behavior of enhanced input action nodes the same as other
+ * event nodes. 
+ */
+UCLASS(Transient)
+class UInputActionEventNodeSpawner final : public UBlueprintNodeSpawner
+{
+	GENERATED_BODY()
+public:
+	static UInputActionEventNodeSpawner* Create(TSubclassOf<UEdGraphNode> const NodeClass, TObjectPtr<const UInputAction> InAction);
+private:
+	virtual UEdGraphNode* Invoke(UEdGraph* ParentGraph, FBindingSet const& Bindings, FVector2D const Location) const override;
+
+	UK2Node* FindExistingNode(const UBlueprint* Blueprint) const;
+
+	// We don't want references to node spawners to be keeping any input action assets from GC
+    // if you unload a plugin for example, so we keep it as a weak pointer.
+	UPROPERTY()
+	TWeakObjectPtr<const UInputAction> WeakActionPtr = nullptr;
+};
+
 UCLASS()
 class INPUTBLUEPRINTNODES_API UK2Node_EnhancedInputAction : public UK2Node, public IK2Node_EventNodeInterface
 {
