[Backout] - CL37584538

Changing this CL to simply clean up comments and TODO's without any behavior changes.
#rnx

Original CL Desc
-----------------------------------------------------------------
[Backout] - CL37579153 because this causes an exception in some cases where the pin is not found during expansion

Original CL Desc
-----------------------------------------------------------------
Clean up some misc. Enhanced Input TODO's and tooltips on BP nodes.

#rnx

diff --git a/Engine/Plugins/EnhancedInput/Source/EnhancedInput/Public/EnhancedInputComponent.h b/Engine/Plugins/EnhancedInput/Source/EnhancedInput/Public/EnhancedInputComponent.h
--- a/Engine/Plugins/EnhancedInput/Source/EnhancedInput/Public/EnhancedInputComponent.h
+++ b/Engine/Plugins/EnhancedInput/Source/EnhancedInput/Public/EnhancedInputComponent.h
@@ -20,6 +20,12 @@ DECLARE_DYNAMIC_DELEGATE_TwoParams(FInputDebugKeyHandlerDynamicSignature, FKey,
 DECLARE_DELEGATE(FEnhancedInputActionHandlerSignature);
 DECLARE_DELEGATE_OneParam(FEnhancedInputActionHandlerValueSignature, const FInputActionValue&);
 DECLARE_DELEGATE_OneParam(FEnhancedInputActionHandlerInstanceSignature, const FInputActionInstance&);	// Provides full access to value and timers
+/**
+ * @param ActionValue		The combined value of all inputs mapped to this action.
+ * @param ElapsedTime		The amount of time in seconds that has elapsed since this input action has begun processing its input values. 
+ * @param TriggeredTime		The amount of time in seconds that this input action has been in the "triggered" state.
+ * @param SourceAction		The input action which caused this event to fire. 
+ */
 DECLARE_DYNAMIC_DELEGATE_FourParams(FEnhancedInputActionHandlerDynamicSignature, FInputActionValue, ActionValue, float, ElapsedTime, float, TriggeredTime, const UInputAction*, SourceAction);
 
 // By default this behavior is on in the editor only. If you would like to turn this
diff --git a/Engine/Plugins/EnhancedInput/Source/InputBlueprintNodes/Private/K2Node_GetInputActionValue.cpp b/Engine/Plugins/EnhancedInput/Source/InputBlueprintNodes/Private/K2Node_GetInputActionValue.cpp
--- a/Engine/Plugins/EnhancedInput/Source/InputBlueprintNodes/Private/K2Node_GetInputActionValue.cpp
+++ b/Engine/Plugins/EnhancedInput/Source/InputBlueprintNodes/Private/K2Node_GetInputActionValue.cpp
@@ -144,7 +144,6 @@ void UK2Node_GetInputActionValue::ExpandNode(class FKismetCompilerContext& Compi
 	TArray<UEdGraphPin*> LinkedTo(ReturnValue->LinkedTo);
 	for (UEdGraphPin* EachLink : LinkedTo)
 	{
-		// TODO: Batch connections by conversion node type so conversion node is only hit once.
 		Schema->TryCreateConnection(AccessorNode->FindPinChecked(UEdGraphSchema_K2::PN_ReturnValue), EachLink);
 	}
 }
diff --git a/Engine/Plugins/EnhancedInput/Source/InputBlueprintNodes/Private/K2Node_InputActionValueAccessor.cpp b/Engine/Plugins/EnhancedInput/Source/InputBlueprintNodes/Private/K2Node_InputActionValueAccessor.cpp
--- a/Engine/Plugins/EnhancedInput/Source/InputBlueprintNodes/Private/K2Node_InputActionValueAccessor.cpp
+++ b/Engine/Plugins/EnhancedInput/Source/InputBlueprintNodes/Private/K2Node_InputActionValueAccessor.cpp
@@ -41,15 +41,4 @@ void UK2Node_InputActionValueAccessor::RegisterDynamicBinding(UDynamicBlueprintB
 	Binding.FunctionNameToBind = GET_FUNCTION_NAME_CHECKED(UEnhancedInputLibrary, GetBoundActionValue);
 
 	InputActionBindingObject->InputActionValueBindings.Add(Binding);
-}
-// TODO: Do we care about unique signatures per input action?
-//
-//FBlueprintNodeSignature UK2Node_InputActionValueAccessor::GetSignature() const
-//{
-//	FBlueprintNodeSignature NodeSignature = Super::GetSignature();
-//	NodeSignature.AddKeyValue(InputAction->GetFName().ToString());
-//
-//	return NodeSignature;
-//}
-
-
+}
\ No newline at end of file
