Add missing "virtual" specifier to IsDataValid

#rnx

diff --git a/Engine/Plugins/EnhancedInput/Source/EnhancedInput/Public/PlayerMappableKeySettings.h b/Engine/Plugins/EnhancedInput/Source/EnhancedInput/Public/PlayerMappableKeySettings.h
--- a/Engine/Plugins/EnhancedInput/Source/EnhancedInput/Public/PlayerMappableKeySettings.h
+++ b/Engine/Plugins/EnhancedInput/Source/EnhancedInput/Public/PlayerMappableKeySettings.h
@@ -22,7 +22,7 @@ public:
 	virtual FName GetMappingName() const { return Name; }
 
 #if WITH_EDITOR
-	EDataValidationResult IsDataValid(class FDataValidationContext& Context) const override;
+	virtual EDataValidationResult IsDataValid(class FDataValidationContext& Context) const override;
 
 	/**
 	 * Get the known mapping names that are current in use. This is a helper function if you want to use a "GetOptions" metadata on a UPROPERTY.
