Modify how input modes are gathered and set on an Input Mapping Context. We now have a "default input mode" for mapping contexts in the project settings, and you can choose to overried that query on the IMC asset itself. This makes it much easier for existing projects with large amounts of data to utilze the new system, as well as provice some more options to make it nice and easy to specify the behavior that you want.

#rnx

diff --git a/Engine/Plugins/EnhancedInput/Source/EnhancedInput/Private/EnhancedInputDeveloperSettings.cpp b/Engine/Plugins/EnhancedInput/Source/EnhancedInput/Private/EnhancedInputDeveloperSettings.cpp
--- a/Engine/Plugins/EnhancedInput/Source/EnhancedInput/Private/EnhancedInputDeveloperSettings.cpp
+++ b/Engine/Plugins/EnhancedInput/Source/EnhancedInput/Private/EnhancedInputDeveloperSettings.cpp
@@ -35,4 +35,6 @@ UEnhancedInputDeveloperSettings::UEnhancedInputDeveloperSettings(const FObjectIn
 	, bShouldLogAllWorldSubsystemInputs(false)
 {
 	PlatformSettings.Initialize(UEnhancedInputPlatformSettings::StaticClass());
+	DefaultMappingContextInputModeQuery = FGameplayTagQuery::BuildQuery(FGameplayTagQueryExpression().AllTagsMatch().AddTag(UE::EnhancedInput::InputMode_Default));
+	DefaultInputMode = FGameplayTagContainer(UE::EnhancedInput::InputMode_Default);
 }
\ No newline at end of file
diff --git a/Engine/Plugins/EnhancedInput/Source/EnhancedInput/Private/EnhancedInputSubsystemInterface.cpp b/Engine/Plugins/EnhancedInput/Source/EnhancedInput/Private/EnhancedInputSubsystemInterface.cpp
--- a/Engine/Plugins/EnhancedInput/Source/EnhancedInput/Private/EnhancedInputSubsystemInterface.cpp
+++ b/Engine/Plugins/EnhancedInput/Source/EnhancedInput/Private/EnhancedInputSubsystemInterface.cpp
@@ -807,7 +807,6 @@ void IEnhancedInputSubsystemInterface::RebuildControlMappings()
 	
 	const FGameplayTagContainer& CurrentInputMode = PlayerInput->GetCurrentInputMode();
 	const bool bInputModeFilteringEnabled = GetDefault<UEnhancedInputDeveloperSettings>()->bEnableInputModeFiltering;
-	const bool bIsCurrentInputModeEmpty = CurrentInputMode.IsEmpty();
 	
 	// Clear existing mappings, but retain the mapping array for later processing
 	TArray<FEnhancedActionKeyMapping> OldMappings(MoveTemp(PlayerInput->EnhancedActionMappings));
@@ -879,8 +878,8 @@ void IEnhancedInputSubsystemInterface::RebuildControlMappings()
 		// Check if this mapping context can be applied in our current input mode on the player
 		// If it can't, then we will not process its mappings..
 		if (bInputModeFilteringEnabled &&
-			MappingContext->bFilterMappingsByQuery &&
-			!MappingContext->InputModeQuery.Matches(CurrentInputMode))
+			MappingContext->ShouldFilterMappingByInputMode() &&
+			!MappingContext->GetInputModeQuery().Matches(CurrentInputMode))
 		{
 			UE_LOG(LogEnhancedInput, Log, TEXT("[%hs] Not applying mappings from IMC '%s' because it does not meet the requirements of the current input mode '%s'"),
 				__func__, *GetNameSafe(MappingContext), *CurrentInputMode.ToString());
diff --git a/Engine/Plugins/EnhancedInput/Source/EnhancedInput/Private/EnhancedPlayerInput.cpp b/Engine/Plugins/EnhancedInput/Source/EnhancedInput/Private/EnhancedPlayerInput.cpp
--- a/Engine/Plugins/EnhancedInput/Source/EnhancedInput/Private/EnhancedPlayerInput.cpp
+++ b/Engine/Plugins/EnhancedInput/Source/EnhancedInput/Private/EnhancedPlayerInput.cpp
@@ -34,6 +34,11 @@ namespace UE
 	}
 }
 
+namespace UE::EnhancedInput
+{
+	UE_DEFINE_GAMEPLAY_TAG_COMMENT(InputMode_Default, "EnhancedInput.Modes.Default", "The default input mode for Enhanced Input");	
+}
+
 UEnhancedPlayerInput::UEnhancedPlayerInput()
 	: Super()
 	, bIsFlushingInputThisFrame(false)
@@ -46,9 +51,14 @@ UEnhancedPlayerInput::UEnhancedPlayerInput()
 		return;
 	}
 
+	const UEnhancedInputDeveloperSettings* Settings = GetDefault<UEnhancedInputDeveloperSettings>();
+	check(Settings);
+	
+	CurrentInputMode = Settings->DefaultInputMode;
+
 	if (UE::Input::EnableDefaultMappingContexts)
 	{
-		for (const FDefaultContextSetting& DefaultContext : GetDefault<UEnhancedInputDeveloperSettings>()->DefaultMappingContexts)
+		for (const FDefaultContextSetting& DefaultContext : Settings->DefaultMappingContexts)
 		{
 			if (const UInputMappingContext* IMC = DefaultContext.InputMappingContext.LoadSynchronous())
 			{
diff --git a/Engine/Plugins/EnhancedInput/Source/EnhancedInput/Private/InputMappingContext.cpp b/Engine/Plugins/EnhancedInput/Source/EnhancedInput/Private/InputMappingContext.cpp
--- a/Engine/Plugins/EnhancedInput/Source/EnhancedInput/Private/InputMappingContext.cpp
+++ b/Engine/Plugins/EnhancedInput/Source/EnhancedInput/Private/InputMappingContext.cpp
@@ -32,6 +32,35 @@ EDataValidationResult UInputMappingContext::IsDataValid(FDataValidationContext&
 }
 #endif	// WITH_EDITOR
 
+bool UInputMappingContext::ShouldFilterMappingByInputMode() const
+{
+	return InputModeFilterOptions != EMappingContextInputModeFilterOptions::DoNotFilter;
+}
+
+FGameplayTagQuery UInputMappingContext::GetInputModeQuery() const
+{
+	switch (InputModeFilterOptions)
+	{
+	case EMappingContextInputModeFilterOptions::UseProjectDefaultQuery:
+		{
+			return GetDefault<UEnhancedInputDeveloperSettings>()->DefaultMappingContextInputModeQuery;	
+		}
+		break;
+	case EMappingContextInputModeFilterOptions::UseCustomQuery:
+		{
+			return InputModeQueryOverride;
+		}
+		break;
+	case EMappingContextInputModeFilterOptions::DoNotFilter:
+	default:
+		{
+			ensureMsgf(false, TEXT("Unexpected filter options %u, returning default"), (uint8)InputModeFilterOptions);
+			return FGameplayTagQuery{};	
+		}
+		break;
+	}
+}
+
 bool UInputMappingContext::ShouldShowInputModeQuery()
 {
 	return GetDefault<UEnhancedInputDeveloperSettings>()->bEnableInputModeFiltering;
diff --git a/Engine/Plugins/EnhancedInput/Source/EnhancedInput/Public/EnhancedInputDeveloperSettings.h b/Engine/Plugins/EnhancedInput/Source/EnhancedInput/Public/EnhancedInputDeveloperSettings.h
--- a/Engine/Plugins/EnhancedInput/Source/EnhancedInput/Public/EnhancedInputDeveloperSettings.h
+++ b/Engine/Plugins/EnhancedInput/Source/EnhancedInput/Public/EnhancedInputDeveloperSettings.h
@@ -4,6 +4,7 @@
 
 #include "Engine/DeveloperSettingsBackedByCVars.h"
 #include "Engine/PlatformSettings.h"
+#include "GameplayTagContainer.h"
 #include "EnhancedInputDeveloperSettings.generated.h"
 
 class UInputMappingContext;
@@ -122,7 +123,7 @@ public:
 	 * @see IEnhancedInputSubsystemInterface::RebuildControlMappings
 	 * @see UEnhancedPlayerInput::CurrentInputMode
 	 */
-	UPROPERTY(config, EditAnywhere, Category = "Enhanced Input")
+	UPROPERTY(config, EditAnywhere, Category = "Enhanced Input|Input Modes")
 	uint8 bEnableInputModeFiltering : 1;
 	
 	/** If true, then the world subsystem will be created. */
@@ -135,4 +136,16 @@ public:
  	 */
 	UPROPERTY(config, EditAnywhere, Category = "Enhanced Input|World Subsystem", meta=(editCondition = "bEnableWorldSubsystem", ConsoleVariable="EnhancedInput.bShouldLogAllWorldSubsystemInputs"))
 	uint8 bShouldLogAllWorldSubsystemInputs : 1;
+
+	/**
+	 * This value will be used to populate the default InputModeQuery on Input Mapping Contexts.
+	 */
+	UPROPERTY(config, EditAnywhere, Category = "Enhanced Input|Input Modes", meta = (EditCondition="bEnableInputModeFiltering"))
+	FGameplayTagQuery DefaultMappingContextInputModeQuery;
+
+	/**
+	 * The default input mode which will be set when a UEnhancedPlayerInput object is created.
+	 */
+	UPROPERTY(config, EditAnywhere, Category = "Enhanced Input|Input Modes", meta = (EditCondition="bEnableInputModeFiltering"))
+	FGameplayTagContainer DefaultInputMode;
 };
diff --git a/Engine/Plugins/EnhancedInput/Source/EnhancedInput/Public/EnhancedPlayerInput.h b/Engine/Plugins/EnhancedInput/Source/EnhancedInput/Public/EnhancedPlayerInput.h
--- a/Engine/Plugins/EnhancedInput/Source/EnhancedInput/Public/EnhancedPlayerInput.h
+++ b/Engine/Plugins/EnhancedInput/Source/EnhancedInput/Public/EnhancedPlayerInput.h
@@ -5,6 +5,7 @@
 #include "GameFramework/PlayerInput.h"
 #include "InputAction.h"
 #include "GameplayTagContainer.h"
+#include "NativeGameplayTags.h"
 
 #include "EnhancedPlayerInput.generated.h"
 
@@ -20,6 +21,15 @@ enum class ETriggerEventInternal : uint8;
 enum class EKeyEvent : uint8;
 class UInputMappingContext;
 
+namespace UE::EnhancedInput
+{
+	/**
+	 * The default input mode of Enhanced Input. Every Input Mapping Context will it's default filtering query set to
+	 * check for this exact tag. 
+	 */
+	ENHANCEDINPUT_API UE_DECLARE_GAMEPLAY_TAG_EXTERN(InputMode_Default);
+};
+
 // ContinuouslyInjectedInputs Map is not managed.
 // Continuous input injections seem to be getting garbage collected and
 // crashing in UObject::ProcessEvent when calling ModifyRaw.
diff --git a/Engine/Plugins/EnhancedInput/Source/EnhancedInput/Public/InputMappingContext.h b/Engine/Plugins/EnhancedInput/Source/EnhancedInput/Public/InputMappingContext.h
--- a/Engine/Plugins/EnhancedInput/Source/EnhancedInput/Public/InputMappingContext.h
+++ b/Engine/Plugins/EnhancedInput/Source/EnhancedInput/Public/InputMappingContext.h
@@ -11,6 +11,33 @@ struct FKey;
 
 class UInputAction;
 
+/**
+ * Options for an input mapping context being filtered based on the current input mode of the player.
+ */
+UENUM()
+enum class EMappingContextInputModeFilterOptions : uint8
+{
+	/**
+	 * This mapping context should use the project's default input mode query.
+	 * 
+	 * @see UEnhancedInputDeveloperSettings::DefaultMappingContextInputModeQuery
+	 * 
+	 * (Project Settings -> Engine -> Enhanced Input -> Input Modes)
+	 */
+	UseProjectDefaultQuery,
+
+	/**
+	 * This mapping context should use a custom input mode query instead of the project default.
+	 */
+	UseCustomQuery,
+
+	/**
+	 * This Input mapping context should not be filtered based on the current mode, effectively ignoring
+	 * the current mode.
+	 */
+	DoNotFilter
+};
+
 /**
 * UInputMappingContext : A collection of key to action mappings for a specific input context
 * Could be used to:
@@ -35,26 +62,42 @@ public:
 	virtual EDataValidationResult IsDataValid(class FDataValidationContext& Context) const override;
 #endif
 
+protected:
+	
 	/**
-	 * If true, then the InputModeQuery must match the current input mode in order for these mappings to be processed.
-	 * If false, then the mappings on this context will not be filtered by the input mode.
-	 * 
-	 * @Note: bEnableInputModeFiltering must be enabled in the UEnhancedInputDeveloperSettings for this to be considered.
-	 */
-	UPROPERTY(BlueprintReadOnly, EditAnywhere, Category = "Input Modes", meta = (EditCondition="ShouldShowInputModeQuery()"))
-	bool bFilterMappingsByQuery = false;
+     * Defines how this input mapping context should be filtered based on the current input mode. 
+     *
+     * Default is Use Project Default Query.
+     * 
+     * @Note: bEnableInputModeFiltering must be enabled in the UEnhancedInputDeveloperSettings for this to be considered.
+     */
+    UPROPERTY(BlueprintReadOnly, EditAnywhere, Category = "Input Modes", meta = (EditCondition="ShouldShowInputModeQuery()"))
+	EMappingContextInputModeFilterOptions InputModeFilterOptions = EMappingContextInputModeFilterOptions::UseProjectDefaultQuery;
 	
 	/**
-	 * Tag Query which will be matched against the current Enhanced Input Subsystem's input mode. 
+	 * Tag Query which will be matched against the current Enhanced Input Subsystem's input mode if bOverrideProjectDefaultInputMode is true. 
 	 *
 	 * If this tag query does not match with the current input mode tag container, then the mappings
 	 * will not be processed.
 	 *
 	 * @Note: bEnableInputModeFiltering must be enabled in the UEnhancedInputDeveloperSettings for this to be considered.
 	 */
-	UPROPERTY(BlueprintReadOnly, EditAnywhere, Category = "Input Modes", meta = (EditCondition="bFilterMappingsByQuery && ShouldShowInputModeQuery()"))
-	FGameplayTagQuery InputModeQuery;
+	UPROPERTY(config, BlueprintReadOnly, EditAnywhere, Category = "Input Modes", meta = (EditConditionHides, EditCondition="ShouldShowInputModeQuery() && InputModeFilterOptions == EMappingContextInputModeFilterOptions::UseCustomQuery"))
+	FGameplayTagQuery InputModeQueryOverride;
+
+public:
+
+	/**
+	* @return True if this mapping context should be filtered based on the current input mode.
+	*/
+	bool ShouldFilterMappingByInputMode() const;
 
+	/**
+	 * @return The tag query which should be used when deciding whether this mapping context should be filterd
+	 * out based on the current input mode or not. 
+	 */
+	FGameplayTagQuery GetInputModeQuery() const;
+	
 	// Localized context descriptor
 	UPROPERTY(BlueprintReadOnly, EditAnywhere, Category = "Description", DisplayName = "Description")
 	FText ContextDescription;
@@ -62,8 +105,6 @@ public:
 	friend class FInputContextDetails;
 	friend class FActionMappingsNodeBuilderEx;
 
-public:
-
 	UFUNCTION()
 	static bool ShouldShowInputModeQuery();
 	
