Add the current input mode to "showdebug enhancedinput"

#rnx

diff --git a/Engine/Plugins/EnhancedInput/Source/EnhancedInput/Private/EnhancedInputSubsystemInterfaceDebug.cpp b/Engine/Plugins/EnhancedInput/Source/EnhancedInput/Private/EnhancedInputSubsystemInterfaceDebug.cpp
--- a/Engine/Plugins/EnhancedInput/Source/EnhancedInput/Private/EnhancedInputSubsystemInterfaceDebug.cpp
+++ b/Engine/Plugins/EnhancedInput/Source/EnhancedInput/Private/EnhancedInputSubsystemInterfaceDebug.cpp
@@ -65,6 +65,9 @@ void IEnhancedInputSubsystemInterface::ShowMappingContextDebugInfo(UCanvas* Canv
 		return;
 	}
 
+	DisplayDebugManager.SetDrawColor(FColor::Cyan);
+	DisplayDebugManager.DrawString(FString::Printf(TEXT("Input Mode: %s"), *PlayerInput->GetCurrentInputMode().ToString()));
+
 	if (PlayerInput->EnhancedActionMappings.Num() + PlayerInput->LastInjectedActions.Num() == 0)
 	{
 		DisplayDebugManager.SetDrawColor(FColor::Orange);
