PR #12518: Add missing macro around retrieval of pressed modifier buttons

diff --git a/Engine/Plugins/EnhancedInput/Source/EnhancedInput/Private/EnhancedPlayerInput.cpp b/Engine/Plugins/EnhancedInput/Source/EnhancedInput/Private/EnhancedPlayerInput.cpp
--- a/Engine/Plugins/EnhancedInput/Source/EnhancedInput/Private/EnhancedPlayerInput.cpp
+++ b/Engine/Plugins/EnhancedInput/Source/EnhancedInput/Private/EnhancedPlayerInput.cpp
@@ -559,8 +559,10 @@ void UEnhancedPlayerInput::EvaluateInputDelegates(const TArray<UInputComponent*>
 
 	// Execute appropriate delegates
 
+#if DEV_ONLY_KEY_BINDINGS_AVAILABLE
 	// Cache modifier key states for debug key bindings
 	const bool bAlt = IsAltPressed(), bCtrl = IsCtrlPressed(), bShift = IsShiftPressed(), bCmd = IsCmdPressed();
+#endif
 
 	// TODO: Process APlayerController::InputComponent only!
 	// Walk the stack, top to bottom, grabbing actions and firing triggered delegates
