Make DetermineHardwareDeviceForActionMapping  take in an IMC asset as well, to make it possible to subclass the mapping context and check some additional metadata on it if users want to check that.

#rnx

diff --git a/Engine/Plugins/EnhancedInput/Source/EnhancedInput/Private/UserSettings/EnhancedInputUserSettings.cpp b/Engine/Plugins/EnhancedInput/Source/EnhancedInput/Private/UserSettings/EnhancedInputUserSettings.cpp
--- a/Engine/Plugins/EnhancedInput/Source/EnhancedInput/Private/UserSettings/EnhancedInputUserSettings.cpp
+++ b/Engine/Plugins/EnhancedInput/Source/EnhancedInput/Private/UserSettings/EnhancedInputUserSettings.cpp
@@ -1397,7 +1397,7 @@ bool UEnhancedInputUserSettings::RegisterKeyMappingsToProfile(UEnhancedPlayerMap
 		FKeyMappingRow& MappingRow = Profile.PlayerMappedKeys.FindOrAdd(MappingName);
 
 		// Determine what Slot and Hardware device this Action Mapping should go in
-		const FHardwareDeviceIdentifier MappingDeviceType = DetermineHardwareDeviceForActionMapping(KeyMapping);
+		const FHardwareDeviceIdentifier MappingDeviceType = DetermineHardwareDeviceForActionMapping(KeyMapping, IMC);
 		
 		FSlotCountToMappingNames& MappingNameToSlotCount = MappingsPerDeviceTypes.FindOrAdd(MappingDeviceType);
 
@@ -1468,7 +1468,7 @@ void UEnhancedInputUserSettings::OnKeyMappingRegistered(FPlayerKeyMapping& Regis
 	// Does nothing by default
 }
 
-FHardwareDeviceIdentifier UEnhancedInputUserSettings::DetermineHardwareDeviceForActionMapping(const FEnhancedActionKeyMapping& ActionMapping) const
+FHardwareDeviceIdentifier UEnhancedInputUserSettings::DetermineHardwareDeviceForActionMapping(const FEnhancedActionKeyMapping& ActionMapping, const UInputMappingContext* IMC) const
 {
 	return FHardwareDeviceIdentifier::Invalid;
 }
diff --git a/Engine/Plugins/EnhancedInput/Source/EnhancedInput/Public/UserSettings/EnhancedInputUserSettings.h b/Engine/Plugins/EnhancedInput/Source/EnhancedInput/Public/UserSettings/EnhancedInputUserSettings.h
--- a/Engine/Plugins/EnhancedInput/Source/EnhancedInput/Public/UserSettings/EnhancedInputUserSettings.h
+++ b/Engine/Plugins/EnhancedInput/Source/EnhancedInput/Public/UserSettings/EnhancedInputUserSettings.h
@@ -704,6 +704,12 @@ protected:
 	 */
 	virtual bool RegisterKeyMappingsToProfile(UEnhancedPlayerMappableKeyProfile& Profile, const UInputMappingContext* IMC);
 
+	UE_DEPRECATED(5.6, "Use DetermineHardwareDeviceForActionMapping which takes in a mapping context instead")
+	virtual FHardwareDeviceIdentifier DetermineHardwareDeviceForActionMapping(const FEnhancedActionKeyMapping& ActionMapping) const final
+	{
+		return DetermineHardwareDeviceForActionMapping(ActionMapping, nullptr);
+	}
+	
 	/**
 	 * Determines the hardware device that an action key mapping is associated with. By default, the hardware will be FHardwareDeviceIdentifier::Invalid.
 	 *
@@ -711,7 +717,7 @@ protected:
 	 * devices based on some metadata on the UPlayerMappableKeySettings of the action mapping, or infer a device based
 	 * on the type of FKey it is.
 	 */
-	virtual FHardwareDeviceIdentifier DetermineHardwareDeviceForActionMapping(const FEnhancedActionKeyMapping& ActionMapping) const;
+	virtual FHardwareDeviceIdentifier DetermineHardwareDeviceForActionMapping(const FEnhancedActionKeyMapping& ActionMapping, const UInputMappingContext* IMC = nullptr) const;
 	
 	/** 
 	* Provides a space for subclasses to add additional debug info if desired.
