PR #12586: Add delegates for when an IMC is added and removed. Make the ActionInstanceData map protected instead of public. Adds some helper functions to dynamically add input triggers and modifiers to an input action, which may be useful if you have some kind of accessibility settings which you need to apply dynamically, or during some tutorial phases where you want to add input modifiers/triggers as the player unlocks new mechanics.

diff --git a/Engine/Plugins/EnhancedInput/Source/EnhancedInput/Private/EnhancedInputSubsystems.cpp b/Engine/Plugins/EnhancedInput/Source/EnhancedInput/Private/EnhancedInputSubsystems.cpp
--- a/Engine/Plugins/EnhancedInput/Source/EnhancedInput/Private/EnhancedInputSubsystems.cpp
+++ b/Engine/Plugins/EnhancedInput/Source/EnhancedInput/Private/EnhancedInputSubsystems.cpp
@@ -100,6 +100,20 @@ void UEnhancedInputLocalPlayerSubsystem::ControlMappingsRebuiltThisFrame()
 	ControlMappingsRebuiltDelegate.Broadcast();
 }
 
+void UEnhancedInputLocalPlayerSubsystem::AddMappingContext(const UInputMappingContext* MappingContext, int32 Priority, const FModifyContextOptions& Options)
+{
+	IEnhancedInputSubsystemInterface::AddMappingContext(MappingContext, Priority, Options);
+	
+	OnMappingContextAdded.Broadcast(MappingContext);
+}
+
+void UEnhancedInputLocalPlayerSubsystem::RemoveMappingContext(const UInputMappingContext* MappingContext, const FModifyContextOptions& Options)
+{
+	IEnhancedInputSubsystemInterface::RemoveMappingContext(MappingContext, Options);
+	
+	OnMappingContextRemoved.Broadcast(MappingContext);
+}
+
 // **************************************************************************************************
 // *
 // * UEnhancedInputWorldSubsystem
diff --git a/Engine/Plugins/EnhancedInput/Source/EnhancedInput/Public/EnhancedInputSubsystems.h b/Engine/Plugins/EnhancedInput/Source/EnhancedInput/Public/EnhancedInputSubsystems.h
--- a/Engine/Plugins/EnhancedInput/Source/EnhancedInput/Public/EnhancedInputSubsystems.h
+++ b/Engine/Plugins/EnhancedInput/Source/EnhancedInput/Public/EnhancedInputSubsystems.h
@@ -32,6 +32,8 @@ public:
 	virtual UEnhancedInputUserSettings* GetUserSettings() const override;
 	virtual void InitalizeUserSettings() override;
 	virtual void ControlMappingsRebuiltThisFrame() override;
+	virtual void AddMappingContext(const UInputMappingContext* MappingContext, int32 Priority, const FModifyContextOptions& Options = FModifyContextOptions()) override;
+	virtual void RemoveMappingContext(const UInputMappingContext* MappingContext, const FModifyContextOptions& Options = FModifyContextOptions()) override;
 protected:
 	virtual TMap<TObjectPtr<const UInputAction>, FInjectedInput>& GetContinuouslyInjectedInputs() override { return ContinuouslyInjectedInputs; }
 	// End IEnhancedInputSubsystemInterface
@@ -46,6 +48,8 @@ public:
 
 	/** A delegate that will be called when control mappings have been rebuilt this frame. */
 	DECLARE_DYNAMIC_MULTICAST_DELEGATE(FOnControlMappingsRebuilt);
+	DECLARE_DYNAMIC_MULTICAST_DELEGATE_OneParam(FOnMappingContextAdded, const UInputMappingContext*, MappingContext);
+	DECLARE_DYNAMIC_MULTICAST_DELEGATE_OneParam(FOnMappingContextRemoved, const UInputMappingContext*, MappingContext);
 
 	/**
 	 * Blueprint Event that is called at the end of any frame that Control Mappings have been rebuilt.
@@ -53,6 +57,24 @@ public:
 	UPROPERTY(BlueprintAssignable, DisplayName=OnControlMappingsRebuilt, Category = "Input")
 	FOnControlMappingsRebuilt ControlMappingsRebuiltDelegate;
 
+	/*
+	* A callback fired when a mapping context is added (AddMappingContext is called on this subsystem)
+	* 
+	* Note:  This does not make any gaurantee that the control mappings will have been rebuilt. 
+	* If you need that, then listen to the ControlMappingsRebuiltDelegate instead.
+	*/
+	UPROPERTY(BlueprintAssignable, Category = "Input")
+	FOnMappingContextAdded OnMappingContextAdded;
+
+	/*
+	* A callback fired when a mapping context is removed (RemoveMappingContext is called on this subsystem)
+	* 
+	* Note:  This does not make any gaurantee that the control mappings will have been rebuilt. 
+	* If you need that, then listen to the ControlMappingsRebuiltDelegate instead.
+	*/
+	UPROPERTY(BlueprintAssignable, Category = "Input")
+	FOnMappingContextRemoved OnMappingContextRemoved;
+
 protected:
     	
 	/** The user settings for this subsystem used to store each user's input related settings */
diff --git a/Engine/Plugins/EnhancedInput/Source/EnhancedInput/Public/InputAction.h b/Engine/Plugins/EnhancedInput/Source/EnhancedInput/Public/InputAction.h
--- a/Engine/Plugins/EnhancedInput/Source/EnhancedInput/Public/InputAction.h
+++ b/Engine/Plugins/EnhancedInput/Source/EnhancedInput/Public/InputAction.h
@@ -261,4 +261,7 @@ public:
 
 	// The source action that this instance is created from
 	const UInputAction* GetSourceAction() const { return SourceAction; }
+
+	void AddInputTrigger(UInputTrigger* InputTrigger) { Triggers.Add(InputTrigger); }
+	void RemoveInputTrigger(UInputTrigger* InputTrigger) { Triggers.Remove(InputTrigger); }
 };
