Fix up some comments on mapping contexts

#rnx

diff --git a/Engine/Plugins/EnhancedInput/Source/EnhancedInput/Public/InputMappingContext.h b/Engine/Plugins/EnhancedInput/Source/EnhancedInput/Public/InputMappingContext.h
--- a/Engine/Plugins/EnhancedInput/Source/EnhancedInput/Public/InputMappingContext.h
+++ b/Engine/Plugins/EnhancedInput/Source/EnhancedInput/Public/InputMappingContext.h
@@ -75,7 +75,7 @@ protected:
 	EMappingContextInputModeFilterOptions InputModeFilterOptions = EMappingContextInputModeFilterOptions::UseProjectDefaultQuery;
 	
 	/**
-	 * Tag Query which will be matched against the current Enhanced Input Subsystem's input mode if bOverrideProjectDefaultInputMode is true. 
+	 * Tag Query which will be matched against the current Enhanced Input Subsystem's input mode if InputModeFilterOptions is set to UseCustomQuery. 
 	 *
 	 * If this tag query does not match with the current input mode tag container, then the mappings
 	 * will not be processed.
@@ -105,6 +105,7 @@ public:
 	friend class FInputContextDetails;
 	friend class FActionMappingsNodeBuilderEx;
 
+	/** UFUNCTION helper to be used as en edit condition for displaying input mode query releated properties. */
 	UFUNCTION()
 	static bool ShouldShowInputModeQuery();
 	
