Misc. EI cleanup:

- Alphabetize some headers
- Remove some deprecated functions
- Deprecate UEnhancedInputDeveloperSettings::bShouldLogAllWorldSubsystemInputs and remove the EnhancedInput.bShouldLogAllWorldSubsystemInputs cvar. This was a pointless setting because the log category could be set to VeryVerbose anyway.

#rnx

diff --git a/Engine/Plugins/EnhancedInput/Source/EnhancedInput/Private/EnhancedActionKeyMapping.cpp b/Engine/Plugins/EnhancedInput/Source/EnhancedInput/Private/EnhancedActionKeyMapping.cpp
--- a/Engine/Plugins/EnhancedInput/Source/EnhancedInput/Private/EnhancedActionKeyMapping.cpp
+++ b/Engine/Plugins/EnhancedInput/Source/EnhancedInput/Private/EnhancedActionKeyMapping.cpp
@@ -1,12 +1,11 @@
 // Copyright Epic Games, Inc. All Rights Reserved.
 
 #include "EnhancedActionKeyMapping.h"
-#include "PlayerMappableKeySettings.h"
+
 #include "HAL/IConsoleManager.h"
-#if WITH_EDITOR
 #include "Internationalization/Text.h"
 #include "Misc/DataValidation.h"
-#endif
+#include "PlayerMappableKeySettings.h"
 
 #include UE_INLINE_GENERATED_CPP_BY_NAME(EnhancedActionKeyMapping)
 
@@ -39,21 +38,6 @@ UPlayerMappableKeySettings* FEnhancedActionKeyMapping::GetPlayerMappableKeySetti
 	return nullptr;
 }
 
-#if WITH_EDITORONLY_DATA
-
-PRAGMA_DISABLE_DEPRECATION_WARNINGS
-FEnhancedActionKeyMapping::FEnhancedActionKeyMapping(const UInputAction* InAction /*= nullptr*/, const FKey InKey /*= EKeys::Invalid*/)
-	: PlayerMappableOptions(InAction)
-	, Action(InAction)
-	, Key(InKey)
-	, bShouldBeIgnored(false)
-	, bHasAlwaysTickTrigger(false)
-	, bIsPlayerMappable(false)
-{}
-PRAGMA_ENABLE_DEPRECATION_WARNINGS
-
-#else
-
 FEnhancedActionKeyMapping::FEnhancedActionKeyMapping(const UInputAction* InAction /*= nullptr*/, const FKey InKey /*= EKeys::Invalid*/)
 	: Action(InAction)
 	, Key(InKey)
@@ -61,10 +45,6 @@ FEnhancedActionKeyMapping::FEnhancedActionKeyMapping(const UInputAction* InActio
 	, bHasAlwaysTickTrigger(false)
 {}
 
-#endif	// WITH_EDITORONLY_DATA
-
-
-
 FName FEnhancedActionKeyMapping::GetMappingName() const
 {
 	if (IsPlayerMappable())
diff --git a/Engine/Plugins/EnhancedInput/Source/EnhancedInput/Private/EnhancedInputDeveloperSettings.cpp b/Engine/Plugins/EnhancedInput/Source/EnhancedInput/Private/EnhancedInputDeveloperSettings.cpp
--- a/Engine/Plugins/EnhancedInput/Source/EnhancedInput/Private/EnhancedInputDeveloperSettings.cpp
+++ b/Engine/Plugins/EnhancedInput/Source/EnhancedInput/Private/EnhancedInputDeveloperSettings.cpp
@@ -2,25 +2,12 @@
 
 #include "EnhancedInputDeveloperSettings.h"
 #include "EnhancedInputPlatformSettings.h"
-#include "UserSettings/EnhancedInputUserSettings.h"
 #include "EnhancedPlayerInput.h"
 #include "HAL/IConsoleManager.h"
+#include "UserSettings/EnhancedInputUserSettings.h"
 
 #include UE_INLINE_GENERATED_CPP_BY_NAME(EnhancedInputDeveloperSettings)
 
-namespace UE::EnhancedInput::Private
-{
-	namespace ConsoleVariables
-	{
-		static bool bShouldLogAllWorldSubsystemInputs = false;
-		static FAutoConsoleVariableRef CVarShouldLogAllWorldSubsystemInputs(
-			TEXT("EnhancedInput.bShouldLogAllWorldSubsystemInputs"),
-			bShouldLogAllWorldSubsystemInputs,
-			TEXT("Should each InputKey call to the World subsystem be logged?"),
-			ECVF_Cheat);
-	}
-}
-
 UEnhancedInputDeveloperSettings::UEnhancedInputDeveloperSettings(const FObjectInitializer& Initializer)
 	: Super(Initializer)
 	, UserSettingsClass(UEnhancedInputUserSettings::StaticClass())
@@ -32,7 +19,6 @@ UEnhancedInputDeveloperSettings::UEnhancedInputDeveloperSettings(const FObjectIn
 	, bShouldOnlyTriggerLastActionInChord(true)
 	, bEnableInputModeFiltering(true)
 	, bEnableWorldSubsystem(false)
-	, bShouldLogAllWorldSubsystemInputs(false)
 {
 	PlatformSettings.Initialize(UEnhancedInputPlatformSettings::StaticClass());
 	DefaultMappingContextInputModeQuery = FGameplayTagQuery::BuildQuery(FGameplayTagQueryExpression().AllTagsMatch().AddTag(UE::EnhancedInput::InputMode_Default));
diff --git a/Engine/Plugins/EnhancedInput/Source/EnhancedInput/Private/EnhancedInputLibrary.cpp b/Engine/Plugins/EnhancedInput/Source/EnhancedInput/Private/EnhancedInputLibrary.cpp
--- a/Engine/Plugins/EnhancedInput/Source/EnhancedInput/Private/EnhancedInputLibrary.cpp
+++ b/Engine/Plugins/EnhancedInput/Source/EnhancedInput/Private/EnhancedInputLibrary.cpp
@@ -2,16 +2,17 @@
 
 #include "EnhancedInputLibrary.h"
 
+#include "Engine/LocalPlayer.h"
 #include "EnhancedActionKeyMapping.h"
 #include "EnhancedInputComponent.h"
+#include "EnhancedInputDeveloperSettings.h"
 #include "EnhancedInputModule.h"
 #include "EnhancedInputSubsystems.h"
-#include "UObject/UObjectIterator.h"
 #include "GameFramework/Actor.h"
-#include "EnhancedInputModule.h"	// For LogEnhancedInput
-#include "EnhancedInputDeveloperSettings.h"
-#include "InputTriggers.h"
 #include "GameFramework/PlayerController.h"
+#include "InputTriggers.h"
+#include "PlayerMappableKeySettings.h"
+#include "UObject/UObjectIterator.h"
 
 #include UE_INLINE_GENERATED_CPP_BY_NAME(EnhancedInputLibrary)
 
@@ -80,12 +81,6 @@ void UEnhancedInputLibrary::BreakInputActionValue(FInputActionValue InActionValu
 	Type = InActionValue.GetValueType();
 }
 
-PRAGMA_DISABLE_DEPRECATION_WARNINGS
-FInputActionValue UEnhancedInputLibrary::MakeInputActionValue(double X, double Y, double Z, const FInputActionValue& MatchValueType)
-{
-	return FInputActionValue(MatchValueType.GetValueType(), FVector(X, Y, Z));
-}
-
 UPlayerMappableKeySettings* UEnhancedInputLibrary::GetPlayerMappableKeySettings(const FEnhancedActionKeyMapping& ActionKeyMapping)
 {
 	return ActionKeyMapping.GetPlayerMappableKeySettings();
@@ -101,8 +96,6 @@ bool UEnhancedInputLibrary::IsActionKeyMappingPlayerMappable(const FEnhancedActi
 	return ActionKeyMapping.IsPlayerMappable();
 }
 
-PRAGMA_ENABLE_DEPRECATION_WARNINGS
-
 FInputActionValue UEnhancedInputLibrary::MakeInputActionValueOfType(double X, double Y, double Z, const EInputActionValueType ValueType)
 {
 	return FInputActionValue(ValueType, FVector(X, Y, Z));
diff --git a/Engine/Plugins/EnhancedInput/Source/EnhancedInput/Private/EnhancedInputSubsystems.cpp b/Engine/Plugins/EnhancedInput/Source/EnhancedInput/Private/EnhancedInputSubsystems.cpp
--- a/Engine/Plugins/EnhancedInput/Source/EnhancedInput/Private/EnhancedInputSubsystems.cpp
+++ b/Engine/Plugins/EnhancedInput/Source/EnhancedInput/Private/EnhancedInputSubsystems.cpp
@@ -1,13 +1,14 @@
 // Copyright Epic Games, Inc. All Rights Reserved.
 
 #include "EnhancedInputSubsystems.h"
-#include "EnhancedInputModule.h"
+
 #include "Components/InputComponent.h"
-#include "Engine/LocalPlayer.h"
 #include "EnhancedInputDeveloperSettings.h"
+#include "EnhancedInputModule.h"
+#include "EnhancedInputWorldProcessor.h"
 #include "Engine/Canvas.h"
+#include "Engine/LocalPlayer.h"
 #include "Engine/World.h"
-#include "EnhancedInputWorldProcessor.h"
 #include "Framework/Application/SlateApplication.h"
 #include "GameFramework/PlayerController.h"
 #include "InputMappingContext.h"
@@ -271,14 +272,7 @@ bool UEnhancedInputWorldSubsystem::InputKey(const FInputKeyParams& Params)
 {
 	if (PlayerInput)
 	{
-
-#if !UE_BUILD_SHIPPING
-		if (GetDefault<UEnhancedInputDeveloperSettings>()->bShouldLogAllWorldSubsystemInputs)
-		{
-			UE_LOG(LogWorldSubsystemInput, VeryVerbose, TEXT("EI %s World Subsystem InputKey : [%s]"), LexToString(GetWorld()->WorldType), *Params.Key.ToString());
-		}
-#endif
-
+		UE_LOG(LogWorldSubsystemInput, VeryVerbose, TEXT("EI %s World Subsystem InputKey : [%s]"), LexToString(GetWorld()->WorldType), *Params.Key.ToString());
 		return PlayerInput->InputKey(Params);
 	}
 
diff --git a/Engine/Plugins/EnhancedInput/Source/EnhancedInput/Private/InputMappingContext.cpp b/Engine/Plugins/EnhancedInput/Source/EnhancedInput/Private/InputMappingContext.cpp
--- a/Engine/Plugins/EnhancedInput/Source/EnhancedInput/Private/InputMappingContext.cpp
+++ b/Engine/Plugins/EnhancedInput/Source/EnhancedInput/Private/InputMappingContext.cpp
@@ -2,24 +2,16 @@
 
 #include "InputMappingContext.h"
 
+#include "EnhancedInputDeveloperSettings.h"
 #include "EnhancedInputLibrary.h"
 #include "EnhancedInputModule.h"
-#include "PlayerMappableKeySettings.h"
-#include "EnhancedInputDeveloperSettings.h"
-
-#if WITH_EDITOR
 #include "Misc/DataValidation.h"
-#endif
+#include "PlayerMappableKeySettings.h"
 
 #include UE_INLINE_GENERATED_CPP_BY_NAME(InputMappingContext)
 
 #define LOCTEXT_NAMESPACE "InputMappingContext"
 
-namespace UE::EnhancedInput
-{
-	static const FName PlayerMappableOptionsHaveBeenUpgradedname = TEXT("__EnhancedInput_Internal_HasBeenUpgraded__");
-}
-
 #if WITH_EDITOR
 EDataValidationResult UInputMappingContext::IsDataValid(FDataValidationContext& Context) const
 {
@@ -66,57 +58,6 @@ bool UInputMappingContext::ShouldShowInputModeQuery()
 	return GetDefault<UEnhancedInputDeveloperSettings>()->bEnableInputModeFiltering;
 }
 
-void UInputMappingContext::PostLoad()
-{
-	Super::PostLoad();
-	
-#if WITH_EDITORONLY_DATA
-	PRAGMA_DISABLE_DEPRECATION_WARNINGS
-	
-	for (FEnhancedActionKeyMapping& Mapping : Mappings)
-	{
-		// If this was player mappable, then set the behavior to override the action and populate it with the same data
-		if (Mapping.bIsPlayerMappable && Mapping.PlayerMappableOptions.Name != UE::EnhancedInput::PlayerMappableOptionsHaveBeenUpgradedname)
-		{
-			Mapping.SettingBehavior = EPlayerMappableKeySettingBehaviors::OverrideSettings;
-
-			// If something was player mappable before then it probably didnt have a default object here
-			if (!Mapping.PlayerMappableKeySettings)
-			{
-				FString NewObjectName;
-				EObjectFlags MaskedOuterFlags = RF_Public | RF_Transactional;
-				Mapping.PlayerMappableKeySettings = NewObject<UPlayerMappableKeySettings>(this, UPlayerMappableKeySettings::StaticClass(), *NewObjectName, MaskedOuterFlags, nullptr);
-			}
-
-			check(Mapping.PlayerMappableKeySettings);
-
-			Mapping.PlayerMappableKeySettings->Metadata = Mapping.PlayerMappableOptions.Metadata;
-
-			// If the name was set already, use that
-			if (Mapping.PlayerMappableKeySettings->Name.IsValid())
-			{
-				Mapping.PlayerMappableKeySettings->Name = Mapping.PlayerMappableOptions.Name;
-			}
-			// Otherwise let's generate something so that the data is valid
-			else
-			{
-				FName UniqueName = MakeUniqueObjectName(this, UInputMappingContext::StaticClass(), Mapping.Action.GetFName());
-				Mapping.PlayerMappableKeySettings->Name = UniqueName;
-			}
-
-			Mapping.PlayerMappableKeySettings->DisplayName = Mapping.PlayerMappableOptions.DisplayName;
-			Mapping.PlayerMappableKeySettings->DisplayCategory = Mapping.PlayerMappableOptions.DisplayCategory;
-
-			// Set the name of the old player mappable options so that we know it has already been upgraded, and we don't do it again
-			// on the next post load.
-			Mapping.PlayerMappableOptions.Name = UE::EnhancedInput::PlayerMappableOptionsHaveBeenUpgradedname;
-		}		
-	}
-	
-	PRAGMA_ENABLE_DEPRECATION_WARNINGS
-#endif	// WITH_EDITORONLY_DATA
-}
-
 FEnhancedActionKeyMapping& UInputMappingContext::MapKey(const UInputAction* Action, FKey ToKey)
 {
 	IEnhancedInputModule::Get().GetLibrary()->RequestRebuildControlMappingsUsingContext(this);
@@ -128,7 +69,7 @@ void UInputMappingContext::UnmapKey(const UInputAction* Action, FKey Key)
 	int32 MappingIdx = Mappings.IndexOfByPredicate([&Action, &Key](const FEnhancedActionKeyMapping& Other) { return Other.Action == Action && Other.Key == Key; });
 	if (MappingIdx != INDEX_NONE)
 	{
-		Mappings.RemoveAtSwap(MappingIdx);	// TODO: Preserve order?
+		Mappings.RemoveAtSwap(MappingIdx);
 		IEnhancedInputModule::Get().GetLibrary()->RequestRebuildControlMappingsUsingContext(this);
 	}
 }
diff --git a/Engine/Plugins/EnhancedInput/Source/EnhancedInput/Private/InputTriggers.cpp b/Engine/Plugins/EnhancedInput/Source/EnhancedInput/Private/InputTriggers.cpp
--- a/Engine/Plugins/EnhancedInput/Source/EnhancedInput/Private/InputTriggers.cpp
+++ b/Engine/Plugins/EnhancedInput/Source/EnhancedInput/Private/InputTriggers.cpp
@@ -1,13 +1,11 @@
 // Copyright Epic Games, Inc. All Rights Reserved.
 
 #include "InputTriggers.h"
+
 #include "EnhancedInputModule.h"
 #include "EnhancedPlayerInput.h"
 #include "HAL/IConsoleManager.h"
-
-#if WITH_EDITOR
 #include "Misc/DataValidation.h"
-#endif
 
 #include UE_INLINE_GENERATED_CPP_BY_NAME(InputTriggers)
 
@@ -482,28 +480,4 @@ EDataValidationResult UInputTriggerCombo::IsDataValid(FDataValidationContext& Co
 }
 #endif // WITH_EDITOR
 
-void UInputTriggerCombo::PostLoad()
-{
-	Super::PostLoad();
-	
-	// start fix up in case it's an old version that just has the cancel input action
-#if WITH_EDITORONLY_DATA
-	PRAGMA_DISABLE_DEPRECATION_WARNINGS
-
-	if (!CancelActions.IsEmpty())
-	{
-		for (const TObjectPtr<const UInputAction>& InputAction : CancelActions)
-		{
-			// use default settings but set cancel action
-			FInputCancelAction InputCancelAction;
-			InputCancelAction.CancelAction = InputAction;
-			InputCancelActions.Add(InputCancelAction);
-		}
-		CancelActions.Empty();
-	}
-		
-	PRAGMA_ENABLE_DEPRECATION_WARNINGS
-#endif // WITH_EDITORONLY_DATA
-}
-
 #undef LOCTEXT_NAMESPACE
diff --git a/Engine/Plugins/EnhancedInput/Source/EnhancedInput/Public/EnhancedActionKeyMapping.h b/Engine/Plugins/EnhancedInput/Source/EnhancedInput/Public/EnhancedActionKeyMapping.h
--- a/Engine/Plugins/EnhancedInput/Source/EnhancedInput/Public/EnhancedActionKeyMapping.h
+++ b/Engine/Plugins/EnhancedInput/Source/EnhancedInput/Public/EnhancedActionKeyMapping.h
@@ -6,12 +6,11 @@
 
 #include "EnhancedActionKeyMapping.generated.h"
 
+enum class EDataValidationResult : uint8;
 class UInputModifier;
 class UInputTrigger;
 class UPlayerMappableKeySettings;
 
-enum class EDataValidationResult : uint8;
-
 /**
 * Defines which Player Mappable Key Setting to use for a Action Key Mapping.
 */
@@ -26,52 +25,6 @@ enum class EPlayerMappableKeySettingBehaviors : uint8
 	IgnoreSettings
 };
 
-PRAGMA_DISABLE_DEPRECATION_WARNINGS
-
-/**
- * A struct that represents player facing mapping options for an action key mapping.
- * Use this to set a unique FName for the mapping option to save it, as well as some FText options
- * for use in UI.
- */
-USTRUCT(BlueprintType)
-struct UE_DEPRECATED(5.3, "FPlayerMappableKeyOptions has been deprecated. Please use UPlayerMappableKeySettings instead.") FPlayerMappableKeyOptions
-{
-	GENERATED_BODY()
-
-public:
-	
-	FPlayerMappableKeyOptions(const UInputAction* InAction = nullptr)
-	{
-		if (InAction)
-		{
-			const FString& ActionName = InAction->GetName();
-			Name = FName(*ActionName);
-			DisplayName = FText::FromString(ActionName);
-		}
-		else
-		{
-			Name = NAME_None;
-			DisplayName = FText::GetEmpty();
-		}
-	};
-
-	/** Metadata that can used to store any other related items to this key mapping such as icons, ability assets, etc. */
-	UPROPERTY(EditAnywhere, BlueprintReadWrite, Category = "Input|PlayerMappable")
-	TObjectPtr<UObject> Metadata = nullptr;
-	
-	/** A unique name for this player mapping to be saved with. */
-	UPROPERTY(EditAnywhere, BlueprintReadWrite, Category = "Input|PlayerMappable")
-	FName Name;
-	
-	/** The localized display name of this key mapping. Use this when displaying the mappings to a user. */
-	UPROPERTY(EditAnywhere, BlueprintReadWrite, Category = "Input|PlayerMappable")
-	FText DisplayName;
-
-	/** The category that this player mapping is in */
-	UPROPERTY(EditAnywhere, BlueprintReadWrite, Category = "Input|PlayerMappable")
-	FText DisplayCategory = FText::GetEmpty();
-};
-
 /**
  * Defines a mapping between a key activation and the resulting enhanced action
  * An key could be a button press, joystick axis movement, etc.
@@ -115,15 +68,6 @@ struct ENHANCEDINPUT_API FEnhancedActionKeyMapping
 	
 	/** Identical comparison, including Triggers and Modifiers current inner values. */
 	bool operator==(const FEnhancedActionKeyMapping& Other) const;
-
-#if WITH_EDITORONLY_DATA
-
-	/** Options for making this a player mappable keymapping */
-	UE_DEPRECATED(5.3, "PlayerMappableOptions has been deprecated, please use PlayerMappableKeySettings instead.")
-	UPROPERTY(EditAnywhere, Category = "Input|PlayerMappable", meta = (editCondition = "bIsPlayerMappable", DisplayAfter = "bIsPlayerMappable", DeprecatedProperty, DeprecationMessage="PlayerMappableOptions has been deprecated, please use the PlayerMappableKeySettings instead"))
-	FPlayerMappableKeyOptions PlayerMappableOptions;
-
-#endif	// WITH_EDITORONLY_DATA
 	
 	/**
 	* Trigger qualifiers. If any trigger qualifiers exist the mapping will not trigger unless:
@@ -171,15 +115,6 @@ struct ENHANCEDINPUT_API FEnhancedActionKeyMapping
 	UPROPERTY(Transient)
 	uint8 bHasAlwaysTickTrigger : 1;
 
-#if WITH_EDITORONLY_DATA
-
-	/** If true then this ActionKeyMapping will be exposed as a player mappable key */
-	UE_DEPRECATED(5.3, "bIsPlayerMappable has been deprecated, please use SettingBehavior instead.")
-	UPROPERTY(EditAnywhere, Category = "Input|PlayerMappable", meta=(DeprecatedProperty, DeprecationMessage="bIsPlayerMappable has been deprecated, please use the SettingBehavior instead"))
-	uint8 bIsPlayerMappable : 1;
-
-#endif	// WITH_EDITORONLY_DATA
-
 protected:
 
 	/**
@@ -263,6 +198,4 @@ public:
 		return true;
 	}
 
-};
-
-PRAGMA_ENABLE_DEPRECATION_WARNINGS
+};
\ No newline at end of file
diff --git a/Engine/Plugins/EnhancedInput/Source/EnhancedInput/Public/EnhancedInputDeveloperSettings.h b/Engine/Plugins/EnhancedInput/Source/EnhancedInput/Public/EnhancedInputDeveloperSettings.h
--- a/Engine/Plugins/EnhancedInput/Source/EnhancedInput/Public/EnhancedInputDeveloperSettings.h
+++ b/Engine/Plugins/EnhancedInput/Source/EnhancedInput/Public/EnhancedInputDeveloperSettings.h
@@ -5,13 +5,13 @@
 #include "Engine/DeveloperSettingsBackedByCVars.h"
 #include "Engine/PlatformSettings.h"
 #include "GameplayTagContainer.h"
+
 #include "EnhancedInputDeveloperSettings.generated.h"
 
 class UInputMappingContext;
 class UEnhancedInputUserSettings;
 class UEnhancedPlayerMappableKeyProfile;
 class UEnhancedPlayerInput;
-enum class EPlayerMappableKeySlot : uint8;
 
 /** Represents a single input mapping context and the priority that it should be applied with */
 USTRUCT()
@@ -134,8 +134,8 @@ public:
  	 * If true then the Enhanced Input world subsystem will log all input that is being processed by it (keypresses, analog values, etc)
  	 * Note: This can produce A LOT of logs, so only use this if you are debugging something. Does nothing in shipping builds
  	 */
-	UPROPERTY(config, EditAnywhere, Category = "Enhanced Input|World Subsystem", meta=(editCondition = "bEnableWorldSubsystem", ConsoleVariable="EnhancedInput.bShouldLogAllWorldSubsystemInputs"))
-	uint8 bShouldLogAllWorldSubsystemInputs : 1;
+	UE_DEPRECATED(5.6, "bShouldLogAllWorldSubsystemInputs is deprecated, use the log category on VeryVerbose for LogWorldSubsystemInput instead.")
+	uint8 bShouldLogAllWorldSubsystemInputs : 1 = false;
 
 	/**
 	 * This value will be used to populate the default InputModeQuery on Input Mapping Contexts.
diff --git a/Engine/Plugins/EnhancedInput/Source/EnhancedInput/Public/EnhancedInputLibrary.h b/Engine/Plugins/EnhancedInput/Source/EnhancedInput/Public/EnhancedInputLibrary.h
--- a/Engine/Plugins/EnhancedInput/Source/EnhancedInput/Public/EnhancedInputLibrary.h
+++ b/Engine/Plugins/EnhancedInput/Source/EnhancedInput/Public/EnhancedInputLibrary.h
@@ -6,13 +6,11 @@
 
 #include "EnhancedInputLibrary.generated.h"
 
-class IEnhancedInputSubsystemInterface;
+class APlayerController;
 enum class EInputActionValueType : uint8;
 struct FEnhancedActionKeyMapping;
 struct FInputActionValue;
-
-class APlayerController;
-class UEnhancedPlayerInput;
+class IEnhancedInputSubsystemInterface;
 class UInputAction;
 class UInputMappingContext;
 class UPlayerMappableKeySettings;
@@ -47,11 +45,7 @@ public:
 	 */
 	UFUNCTION(BlueprintPure, Category = "Input", meta = (Keywords = "construct build", NativeMakeFunc))
 	static FInputActionValue MakeInputActionValueOfType(double X, double Y, double Z, const EInputActionValueType ValueType);
-
-	UE_DEPRECATED(5.1, "This version of MakeInputActionValue has been deprecated, please use MakeInputActionValueOfType")
-	UFUNCTION(BlueprintPure, Category = "Input", meta = (Keywords = "construct build", NativeMakeFunc, DeprecatedFunction, DeprecatedMessage="This version of MakeInputActionValue has been deprecated, please use MakeInputActionValueOfType"))
-	static FInputActionValue MakeInputActionValue(double X, double Y, double Z, const FInputActionValue& MatchValueType);
-
+	
 	/**
 	* Returns the Player Mappable Key Settings owned by the Action Key Mapping or by the referenced Input Action, or nothing based of the Setting Behavior.
 	*/
diff --git a/Engine/Plugins/EnhancedInput/Source/EnhancedInput/Public/InputMappingContext.h b/Engine/Plugins/EnhancedInput/Source/EnhancedInput/Public/InputMappingContext.h
--- a/Engine/Plugins/EnhancedInput/Source/EnhancedInput/Public/InputMappingContext.h
+++ b/Engine/Plugins/EnhancedInput/Source/EnhancedInput/Public/InputMappingContext.h
@@ -108,8 +108,6 @@ public:
 	/** UFUNCTION helper to be used as en edit condition for displaying input mode query releated properties. */
 	UFUNCTION()
 	static bool ShouldShowInputModeQuery();
-	
-	virtual void PostLoad() override;
 
 	/**
 	* Mapping accessors.
@@ -131,10 +129,6 @@ public:
 	*/
 	UFUNCTION(BlueprintCallable, Category = "Mapping")
 	void UnmapKey(const UInputAction* Action, FKey Key);
-
-	UE_DEPRECATED(5.1, "UnmapAction has been deprecated, please use UnmapAllKeysFromAction instead.")
-	UFUNCTION(BlueprintCallable, Category = "Mapping", meta=(DeprecatedFunction, DeprecatedMessage="UnmapAction has been deprecated, please use UnmapAllKeysFromAction instead."))
-	void UnmapAction(const UInputAction* Action) { UnmapAllKeysFromAction(Action); }
 	
 	/**
 	* Unmap all key maps to an action within the mapping context.
diff --git a/Engine/Plugins/EnhancedInput/Source/EnhancedInput/Public/InputTriggers.h b/Engine/Plugins/EnhancedInput/Source/EnhancedInput/Public/InputTriggers.h
--- a/Engine/Plugins/EnhancedInput/Source/EnhancedInput/Public/InputTriggers.h
+++ b/Engine/Plugins/EnhancedInput/Source/EnhancedInput/Public/InputTriggers.h
@@ -550,14 +550,6 @@ public:
 #if WITH_EDITOR
 	virtual EDataValidationResult IsDataValid(class FDataValidationContext& Context) const override;
 #endif
-
-#if  WITH_EDITORONLY_DATA
-	// Actions that will cancel the combo if they are triggered.
-	UE_DEPRECATED(5.2, "CancelActions has been deprecated as of 5.2. Please use InputCancelActions instead.")
-	TArray<TObjectPtr<const UInputAction>> CancelActions;
-#endif
-
-	virtual void PostLoad() override;
 	
 	/**
 	 * List of input actions that need to be completed (according to Combo Step Completion States) to trigger this action.
diff --git a/Engine/Plugins/EnhancedInput/Source/InputEditor/Private/InputEditorModule.cpp b/Engine/Plugins/EnhancedInput/Source/InputEditor/Private/InputEditorModule.cpp
--- a/Engine/Plugins/EnhancedInput/Source/InputEditor/Private/InputEditorModule.cpp
+++ b/Engine/Plugins/EnhancedInput/Source/InputEditor/Private/InputEditorModule.cpp
@@ -311,18 +311,6 @@ void FAssetTypeActions_InputAction::GetCreateContextFromActionsMenu(TArray<TWeak
 	}
 }
 
-PRAGMA_DISABLE_DEPRECATION_WARNINGS
-class UE_DEPRECATED(5.3, "PlayerMappableInputConfig has been deprecated, please use UEnhancedInputUserSettings instead") FAssetTypeActions_PlayerMappableInputConfig : public FAssetTypeActions_DataAsset
-{
-public:
-	virtual FText GetName() const override { return NSLOCTEXT("AssetTypeActions", "AssetTypeActions_PlayerMappableInputConfig", "Player Mappable Input Config (Deprecated)"); }
-	virtual uint32 GetCategories() override { return FInputEditorModule::GetInputAssetsCategory(); }
-	virtual FColor GetTypeColor() const override { return FColor(127, 255, 255); }
-	virtual FText GetAssetDescription(const FAssetData& AssetData) const override { return NSLOCTEXT("AssetTypeActions", "AssetTypeActions_PlayerBindableInputConfigDesc", "Represents one set of Player Mappable controller/keymappings"); }
-	virtual UClass* GetSupportedClass() const override { return UPlayerMappableInputConfig::StaticClass(); }
-};
-PRAGMA_ENABLE_DEPRECATION_WARNINGS
-
 /** Custom style set for Enhanced Input */
 class FEnhancedInputSlateStyle final : public FSlateStyleSet
 {
@@ -372,9 +360,6 @@ void FInputEditorModule::StartupModule()
 	{
 		RegisterAssetTypeActions(AssetTools, MakeShareable(new FAssetTypeActions_InputAction));
 		RegisterAssetTypeActions(AssetTools, MakeShareable(new FAssetTypeActions_InputContext));
-		PRAGMA_DISABLE_DEPRECATION_WARNINGS
-		RegisterAssetTypeActions(AssetTools, MakeShareable(new FAssetTypeActions_PlayerMappableInputConfig));
-		PRAGMA_ENABLE_DEPRECATION_WARNINGS
 	}
 
 	// Make a new style set for Enhanced Input, which will register any custom icons for the types in this plugin
