Do not initalize the user settings on the enhanced input local player subsystem initalize. Instead, we will rely on the Player Controller change delegate. This Player Controller change delegate was being fired already, so initalizing the user settings on the subsystem's initalize just calls the function twice when we don't need to.

This also ensures that, when we initalize the user settings, there is a valid outer player controller. This will make it easier to do things like handle changing your settings class if you need to, and is a nice assumption to be able to make when loading the save game file.

diff --git a/Engine/Plugins/EnhancedInput/Source/EnhancedInput/Private/EnhancedInputSubsystems.cpp b/Engine/Plugins/EnhancedInput/Source/EnhancedInput/Private/EnhancedInputSubsystems.cpp
--- a/Engine/Plugins/EnhancedInput/Source/EnhancedInput/Private/EnhancedInputSubsystems.cpp
+++ b/Engine/Plugins/EnhancedInput/Source/EnhancedInput/Private/EnhancedInputSubsystems.cpp
@@ -24,16 +24,6 @@ DEFINE_LOG_CATEGORY(LogWorldSubsystemInput);
 // *
 // **************************************************************************************************
 
-void UEnhancedInputLocalPlayerSubsystem::Initialize(FSubsystemCollectionBase& Collection)
-{
-	Super::Initialize(Collection);
-	
-	if (GetDefault<UEnhancedInputDeveloperSettings>()->bEnableUserSettings)
-	{
-		InitalizeUserSettings();
-	}
-}
-
 void UEnhancedInputLocalPlayerSubsystem::PlayerControllerChanged(APlayerController* NewPlayerController)
 {
 	Super::PlayerControllerChanged(NewPlayerController);
diff --git a/Engine/Plugins/EnhancedInput/Source/EnhancedInput/Public/EnhancedInputSubsystems.h b/Engine/Plugins/EnhancedInput/Source/EnhancedInput/Public/EnhancedInputSubsystems.h
--- a/Engine/Plugins/EnhancedInput/Source/EnhancedInput/Public/EnhancedInputSubsystems.h
+++ b/Engine/Plugins/EnhancedInput/Source/EnhancedInput/Public/EnhancedInputSubsystems.h
@@ -23,7 +23,6 @@ class ENHANCEDINPUT_API UEnhancedInputLocalPlayerSubsystem : public ULocalPlayer
 public:
 
 	// Begin ULocalPlayerSubsystem
-	virtual void Initialize(FSubsystemCollectionBase& Collection) override;
 	virtual void PlayerControllerChanged(APlayerController* NewPlayerController) override;
 	// End ULocalPlayerSubsystem
 	
