Fix for the Default Mapping Contexts not being registered with the settings and clean them up to be more consistently applied how other IMC's are, instead of specially in the UEnhancedPlayerInput constructor.

diff --git a/Engine/Plugins/EnhancedInput/Source/EnhancedInput/Private/EnhancedInputSubsystems.cpp b/Engine/Plugins/EnhancedInput/Source/EnhancedInput/Private/EnhancedInputSubsystems.cpp
--- a/Engine/Plugins/EnhancedInput/Source/EnhancedInput/Private/EnhancedInputSubsystems.cpp
+++ b/Engine/Plugins/EnhancedInput/Source/EnhancedInput/Private/EnhancedInputSubsystems.cpp
@@ -24,6 +24,15 @@ DEFINE_LOG_CATEGORY(LogWorldSubsystemInput);
 // *
 // **************************************************************************************************
 
+
+namespace UE::Input
+{
+	static int32 EnableDefaultMappingContexts = 1;
+	static FAutoConsoleVariableRef CVarEnableDefaultMappingContexts(TEXT("EnhancedInput.EnableDefaultMappingContexts"),
+		EnableDefaultMappingContexts,
+		TEXT("Should the UEnhancedInputDeveloperSettings::DefaultMappingContexts be applied to every UEnhancedPlayerInput?"));
+}
+
 void UEnhancedInputLocalPlayerSubsystem::PlayerControllerChanged(APlayerController* NewPlayerController)
 {
 	Super::PlayerControllerChanged(NewPlayerController);
@@ -32,6 +41,26 @@ void UEnhancedInputLocalPlayerSubsystem::PlayerControllerChanged(APlayerControll
 	{
 		InitalizeUserSettings();
 	}
+
+	if (UE::Input::EnableDefaultMappingContexts)
+	{
+		// Add any default input mapping contexts that we may need
+    	for (const FDefaultContextSetting& DefaultMapping : GetDefault<UEnhancedInputDeveloperSettings>()->DefaultMappingContexts)
+    	{
+    		if (const UInputMappingContext* IMC = DefaultMapping.InputMappingContext.LoadSynchronous())
+    		{
+    			if (DefaultMapping.bAddImmediately)
+    			{
+    				AddMappingContext(IMC, DefaultMapping.Priority);
+    			}
+    			
+    			if (DefaultMapping.bRegisterWithUserSettings && UserSettings)
+    			{
+    				UserSettings->RegisterInputMappingContext(IMC);
+    			}
+    		}
+    	}	
+	}	
 }
 
 UEnhancedPlayerInput* UEnhancedInputLocalPlayerSubsystem::GetPlayerInput() const
diff --git a/Engine/Plugins/EnhancedInput/Source/EnhancedInput/Private/EnhancedPlayerInput.cpp b/Engine/Plugins/EnhancedInput/Source/EnhancedInput/Private/EnhancedPlayerInput.cpp
--- a/Engine/Plugins/EnhancedInput/Source/EnhancedInput/Private/EnhancedPlayerInput.cpp
+++ b/Engine/Plugins/EnhancedInput/Source/EnhancedInput/Private/EnhancedPlayerInput.cpp
@@ -22,11 +22,6 @@ namespace UE
 			ShouldOnlyTriggerLastActionInChord,
 			TEXT("Should only the last action in a ChordedAction trigger be fired? If this is disabled, then the dependant chords will be fired as well"));
 
-		static int32 EnableDefaultMappingContexts = 1;
-		static FAutoConsoleVariableRef CVarEnableDefaultMappingContexts(TEXT("EnhancedInput.EnableDefaultMappingContexts"),
-			EnableDefaultMappingContexts,
-			TEXT("Should the UEnhancedInputDeveloperSettings::DefaultMappingContexts be applied to every UEnhancedPlayerInput?"));
-
 		static int32 ReconcileRemovedMappingDelegates = 1;
 		static FAutoConsoleVariableRef CVarReconcileRemovedMappingDelegates(TEXT("EnhancedInput.ReconcileRemovedMappingDelegates"),
 			ReconcileRemovedMappingDelegates,
@@ -55,20 +50,6 @@ UEnhancedPlayerInput::UEnhancedPlayerInput()
 	check(Settings);
 	
 	CurrentInputMode = Settings->DefaultInputMode;
-
-	if (UE::Input::EnableDefaultMappingContexts)
-	{
-		for (const FDefaultContextSetting& DefaultContext : Settings->DefaultMappingContexts)
-		{
-			if (const UInputMappingContext* IMC = DefaultContext.InputMappingContext.LoadSynchronous())
-			{
-				if (DefaultContext.bAddImmediately)
-				{
-					AppliedInputContexts.Add(IMC, DefaultContext.Priority);	
-				}
-			}
-		}
-	}
 }
 
 // NOTE: Enum order represents firing priority(lowest to highest) and is important as multiple keys bound to the same action may generate differing trigger event states.
