Add an upgrade path to Enhanced input user settings.

- The settings file name that is used when saving the Enhanced Input User Settings is now a config property in the developer settings.
- You can override the UEnhancedPlayerInput::GetUserSettingsSaveFileName function to add any custom file naming schemas you may desire
- There is now a virtual HandleSettingsClassChangedOnLoad_Impl function on the enhanced input user settings class which provides an extension point for when the settings class is loaded, but is different then the one you need to create. You can use this function to handle any data retention you may want. By default this does nothing.

diff --git a/Engine/Plugins/EnhancedInput/Source/EnhancedInput/Private/EnhancedPlayerInput.cpp b/Engine/Plugins/EnhancedInput/Source/EnhancedInput/Private/EnhancedPlayerInput.cpp
--- a/Engine/Plugins/EnhancedInput/Source/EnhancedInput/Private/EnhancedPlayerInput.cpp
+++ b/Engine/Plugins/EnhancedInput/Source/EnhancedInput/Private/EnhancedPlayerInput.cpp
@@ -1,32 +1,31 @@
 // Copyright Epic Games, Inc. All Rights Reserved.
 
 #include "EnhancedPlayerInput.h"
+
+#include "Engine/LocalPlayer.h"
+#include "Engine/World.h"
 #include "EnhancedInputComponent.h"
+#include "EnhancedInputDeveloperSettings.h"
+#include "EnhancedInputModule.h"
 #include "GameFramework/PlayerController.h"
-#include "HAL/IConsoleManager.h"
-#include "Engine/World.h"
 #include "GameFramework/WorldSettings.h"
-#include "EnhancedInputDeveloperSettings.h"
+#include "HAL/IConsoleManager.h"
 #include "InputMappingContext.h"
 #include "UserSettings/EnhancedInputUserSettings.h"
-#include "EnhancedInputModule.h"
 
 #include UE_INLINE_GENERATED_CPP_BY_NAME(EnhancedPlayerInput)
 
-namespace UE
+namespace UE::Input
 {
-	namespace Input
-	{
-		static int32 ShouldOnlyTriggerLastActionInChord = 1;
-		static FAutoConsoleVariableRef CVarShouldOnlyTriggerLastActionInChord(TEXT("EnhancedInput.OnlyTriggerLastActionInChord"),
-			ShouldOnlyTriggerLastActionInChord,
-			TEXT("Should only the last action in a ChordedAction trigger be fired? If this is disabled, then the dependant chords will be fired as well"));
-
-		static int32 ReconcileRemovedMappingDelegates = 1;
-		static FAutoConsoleVariableRef CVarReconcileRemovedMappingDelegates(TEXT("EnhancedInput.ReconcileRemovedMappingDelegates"),
-			ReconcileRemovedMappingDelegates,
-			TEXT("When the mappings are rebuilt, they have been 'in process'. If this is true, we will set their value to zero so that they fire the 'Canceled' event on their next evaluation."));
-	}
+	static int32 ShouldOnlyTriggerLastActionInChord = 1;
+	static FAutoConsoleVariableRef CVarShouldOnlyTriggerLastActionInChord(TEXT("EnhancedInput.OnlyTriggerLastActionInChord"),
+		ShouldOnlyTriggerLastActionInChord,
+		TEXT("Should only the last action in a ChordedAction trigger be fired? If this is disabled, then the dependant chords will be fired as well"));
+
+	static int32 ReconcileRemovedMappingDelegates = 1;
+	static FAutoConsoleVariableRef CVarReconcileRemovedMappingDelegates(TEXT("EnhancedInput.ReconcileRemovedMappingDelegates"),
+		ReconcileRemovedMappingDelegates,
+		TEXT("When the mappings are rebuilt, they have been 'in process'. If this is true, we will set their value to zero so that they fire the 'Canceled' event on their next evaluation."));
 }
 
 namespace UE::EnhancedInput
@@ -278,6 +277,15 @@ const FGameplayTagContainer& UEnhancedPlayerInput::GetCurrentInputMode() const
 	return CurrentInputMode;
 }
 
+FString UEnhancedPlayerInput::GetUserSettingsSaveFileName() const
+{
+	// By default, return the string specified in the settings.
+	// You can override this function if you have any more complex scenarios, such as perhaps
+	// doing some kind of custom versioning/upgrading of your settings as your changes over time
+	// or creating unique save slots for different players.
+	return GetDefault<UEnhancedInputDeveloperSettings>()->InputSettingsSaveSlotName;
+}
+
 void UEnhancedPlayerInput::SetCurrentInputMode(const FGameplayTagContainer& NewMode)
 {
 	CurrentInputMode = NewMode;
diff --git a/Engine/Plugins/EnhancedInput/Source/EnhancedInput/Private/UserSettings/EnhancedInputUserSettings.cpp b/Engine/Plugins/EnhancedInput/Source/EnhancedInput/Private/UserSettings/EnhancedInputUserSettings.cpp
--- a/Engine/Plugins/EnhancedInput/Source/EnhancedInput/Private/UserSettings/EnhancedInputUserSettings.cpp
+++ b/Engine/Plugins/EnhancedInput/Source/EnhancedInput/Private/UserSettings/EnhancedInputUserSettings.cpp
@@ -31,9 +31,6 @@
 
 namespace UE::EnhancedInput
 {
-	/** The name of the slot that these settings will save to */
-	static const FString SETTINGS_SLOT_NAME = TEXT("EnhancedInputUserSettings");
-
 	UE_DEFINE_GAMEPLAY_TAG_STATIC(TAG_DefaultProfileIdentifier, "InputUserSettings.Profiles.Default");
 	static const FText DefaultProfileDisplayName = LOCTEXT("Default_Profile_name", "Default Profile");
 	
@@ -625,11 +622,13 @@ UEnhancedInputUserSettings* UEnhancedInputUserSettings::LoadOrCreateSettings(ULo
 		UE_LOG(LogEnhancedInput, Log, TEXT("Unable to determine an owning Local Player for the given Enhanced Player Input object"));
 		return nullptr;
 	}
+
+	const FString& SaveFileName = GetSaveFilename(LocalPlayer);
 	
 	// If the save game exists, load it.
-	if (UGameplayStatics::DoesSaveGameExist(UE::EnhancedInput::SETTINGS_SLOT_NAME, LocalPlayer->GetLocalPlayerIndex()))
+	if (UGameplayStatics::DoesSaveGameExist(SaveFileName, LocalPlayer->GetLocalPlayerIndex()))
 	{
-		USaveGame* Slot = UGameplayStatics::LoadGameFromSlot(UE::EnhancedInput::SETTINGS_SLOT_NAME, LocalPlayer->GetLocalPlayerIndex());
+		USaveGame* Slot = UGameplayStatics::LoadGameFromSlot(SaveFileName, LocalPlayer->GetLocalPlayerIndex());
 		Settings = Cast<UEnhancedInputUserSettings>(Slot);
 	}
 
@@ -644,18 +643,31 @@ UEnhancedInputUserSettings* UEnhancedInputUserSettings::LoadOrCreateSettings(ULo
 		return nullptr;
 	}
 	
-	UE_CLOG((Settings && (Settings->GetClass() != SettingsClass)),
-		LogEnhancedInput,
-		Log,
-		TEXT("[%hs] Enhanced Input User Settings class has been changed from '%s' to '%s'. A new settings save object will be created (your saved settings will be reset)"),
-		__func__, *GetNameSafe(Settings->GetClass()), *GetNameSafe(SettingsClass));
-	
 	// If the settings are null (they dont exist yet) or the class has changed, we need to create a new object.
-	// The class can change if you modify it in the editor to be something else
-	if (Settings == nullptr || (Settings->GetClass() != SettingsClass))
+	if (Settings == nullptr)
 	{
 		Settings = Cast<UEnhancedInputUserSettings>(UGameplayStatics::CreateSaveGameObject(SettingsClass));
 	}
+	// If the class has changed, then we should call into a customizable upgrade path notifying you of the change.
+	// This can happen if you change the class in the editor and there is still an old save file on disk,
+	// or if you change the class of your user settings object and would to keep some of the old data.
+	else if (Settings->GetClass() != SettingsClass)
+	{
+		UE_LOG(
+			LogEnhancedInput,
+			Log,
+			TEXT("[%hs] Enhanced Input User Settings class has been changed from '%s' to '%s'. A new settings save object will be created (your saved settings will be reset). Override HandleSettingsClassChangedOnLoad_Impl to retain data."),
+			__func__, *GetNameSafe(Settings->GetClass()), *GetNameSafe(SettingsClass))
+
+		// Create a new settings object here with the new class which it has been changed to
+		UEnhancedInputUserSettings* NewSettingsObject = Cast<UEnhancedInputUserSettings>(UGameplayStatics::CreateSaveGameObject(SettingsClass));
+
+		// Allow you to upgrade the old settings if you ever change the class of your user settings type
+		NewSettingsObject->HandleSettingsClassChangedOnLoad(Settings);
+
+		// Store our new settings object as the return value
+		Settings = NewSettingsObject;
+	}
 
 	if (ensure(Settings))
 	{
@@ -704,7 +716,7 @@ void UEnhancedInputUserSettings::SaveSettings()
 {
 	if (ULocalPlayer* OwningPlayer = GetLocalPlayer())
 	{
-		UGameplayStatics::SaveGameToSlot(this, UE::EnhancedInput::SETTINGS_SLOT_NAME, OwningPlayer->GetLocalPlayerIndex());
+		UGameplayStatics::SaveGameToSlot(this, GetSaveFilename(OwningPlayer), OwningPlayer->GetLocalPlayerIndex());
 		UE_LOG(LogEnhancedInput, Verbose, TEXT("Enhanced Input User Settings saved!"));
 	}
 	else
@@ -720,7 +732,7 @@ void UEnhancedInputUserSettings::AsyncSaveSettings()
 		FAsyncSaveGameToSlotDelegate SavedDelegate;
 		SavedDelegate.BindUObject(this, &UEnhancedInputUserSettings::OnAsyncSaveComplete);
 		
-		UGameplayStatics::AsyncSaveGameToSlot(this, UE::EnhancedInput::SETTINGS_SLOT_NAME, OwningPlayer->GetLocalPlayerIndex(), SavedDelegate);
+		UGameplayStatics::AsyncSaveGameToSlot(this, GetSaveFilename(OwningPlayer), OwningPlayer->GetLocalPlayerIndex(), SavedDelegate);
 	}
 	else
 	{
@@ -728,6 +740,41 @@ void UEnhancedInputUserSettings::AsyncSaveSettings()
 	}
 }
 
+void UEnhancedInputUserSettings::HandleSettingsClassChangedOnLoad(UEnhancedInputUserSettings* OldSettingsObject)
+{
+	// This function is called when the user settings class has changed
+	check(OldSettingsObject);
+	
+	// This should only ever be called for settings objects of two different classes
+	check(OldSettingsObject->GetClass() != GetClass());
+
+	// Call the overridable virtual function to allow games to customize this behavior
+	HandleSettingsClassChangedOnLoad_Impl(OldSettingsObject);
+}
+
+void UEnhancedInputUserSettings::HandleSettingsClassChangedOnLoad_Impl(UEnhancedInputUserSettings* OldSettingsObject)
+{
+	// Copy over any values of properties from the old settings object to this new one for initialization.
+	// This will keep things like the key mappings, profile id, and other UPROPERTY's that are the same across
+	// the two objects.
+	CurrentProfileIdentifier = OldSettingsObject->CurrentProfileIdentifier;
+	SavedKeyProfiles = OldSettingsObject->SavedKeyProfiles;
+
+	// Note: Calling any Serialize functions, or something like CopyPropertiesForUnrelatedObjects, will cause
+	// errors to be logged because there will be a size mismatch between the two types, as they are different classes.
+	// This is why we naively assign the data we would like to keep above instead of calling a more complex serialization function.
+}
+
+FString UEnhancedInputUserSettings::GetSaveFilename(const ULocalPlayer* LP)
+{
+	check(LP);
+
+	UEnhancedPlayerInput* PlayerInput = LP->GetSubsystem<UEnhancedInputLocalPlayerSubsystem>()->GetPlayerInput();
+	check(PlayerInput);
+
+	return PlayerInput->GetUserSettingsSaveFileName();
+}
+
 void UEnhancedInputUserSettings::OnAsyncSaveComplete(const FString& SlotName, const int32 UserIndex, bool bSuccess)
 {
 	if (bSuccess)
diff --git a/Engine/Plugins/EnhancedInput/Source/EnhancedInput/Public/EnhancedInputDeveloperSettings.h b/Engine/Plugins/EnhancedInput/Source/EnhancedInput/Public/EnhancedInputDeveloperSettings.h
--- a/Engine/Plugins/EnhancedInput/Source/EnhancedInput/Public/EnhancedInputDeveloperSettings.h
+++ b/Engine/Plugins/EnhancedInput/Source/EnhancedInput/Public/EnhancedInputDeveloperSettings.h
@@ -80,6 +80,17 @@ public:
 	 */
 	UPROPERTY(config, EditAnywhere, NoClear, Category = "Enhanced Input|User Settings", meta=(editCondition = "bEnableUserSettings"))
 	TSoftClassPtr<UEnhancedPlayerMappableKeyProfile> DefaultPlayerMappableKeyProfileClass;
+
+	/**
+	 * The default SaveGame slot name to use when loading/saving Enhanced Input user settings.
+	 * 
+	 * The file name will be this string with ".sav" appended to the end of it. So, on disk, the file would be
+	 * "<InputSettingsSaveSlotName>.sav".
+	 * 
+	 * This file will be written to the "SaveGames" directory, which is "<Project Folder>/Saved/SaveGames" in the editor.
+	 */
+	UPROPERTY(config, EditAnywhere, NoClear, Category = "Enhanced Input|User Settings", meta=(editCondition = "bEnableUserSettings"))
+	FString InputSettingsSaveSlotName = TEXT("EnhancedInputUserSettings");
 	
 	/** The default player input class that the Enhanced Input world subsystem will use. */
 	UPROPERTY(config, EditAnywhere, NoClear, Category = "Enhanced Input|World Subsystem", meta=(editCondition = "bEnableWorldSubsystem"))
diff --git a/Engine/Plugins/EnhancedInput/Source/EnhancedInput/Public/EnhancedPlayerInput.h b/Engine/Plugins/EnhancedInput/Source/EnhancedInput/Public/EnhancedPlayerInput.h
--- a/Engine/Plugins/EnhancedInput/Source/EnhancedInput/Public/EnhancedPlayerInput.h
+++ b/Engine/Plugins/EnhancedInput/Source/EnhancedInput/Public/EnhancedPlayerInput.h
@@ -114,6 +114,12 @@ public:
 	 * Returns a const ref to the current input mode. 
 	 */
 	const FGameplayTagContainer& GetCurrentInputMode() const;
+
+	/**
+	 * @return The name which should be used to save the input settings to.
+	 * By default, this will return InputSettingsSaveSlotName from the developer settings.
+	 */
+	virtual FString GetUserSettingsSaveFileName() const;
 	
 protected:
 
diff --git a/Engine/Plugins/EnhancedInput/Source/EnhancedInput/Public/UserSettings/EnhancedInputUserSettings.h b/Engine/Plugins/EnhancedInput/Source/EnhancedInput/Public/UserSettings/EnhancedInputUserSettings.h
--- a/Engine/Plugins/EnhancedInput/Source/EnhancedInput/Public/UserSettings/EnhancedInputUserSettings.h
+++ b/Engine/Plugins/EnhancedInput/Source/EnhancedInput/Public/UserSettings/EnhancedInputUserSettings.h
@@ -512,8 +512,29 @@ public:
 	UFUNCTION(BlueprintCallable, Category="Enhanced Input|User Settings")
 	virtual void AsyncSaveSettings();
 
+private:
+	
+	/**
+	 * Called when a save file is loaded from disk but has a different class then the one specified in the developer settings.
+	 * 
+	 * @param OldSettingsObject Pointer to the old user settings object
+	 */
+	void HandleSettingsClassChangedOnLoad(UEnhancedInputUserSettings* OldSettingsObject);
+	
 protected:
+	
+	/**
+	 * Override this function to be notified of when the user settings class has changed.
+	 * 
+	 * @param OldSettingsObject Pointer to the old user settings object
+	 */
+	virtual void HandleSettingsClassChangedOnLoad_Impl(UEnhancedInputUserSettings* OldSettingsObject);
 
+	/**
+	* Returns the save file name to use when saving the settings file for the given local player
+	*/
+	static FString GetSaveFilename(const ULocalPlayer* LP);
+	
 	virtual void OnAsyncSaveComplete(const FString& SlotName, const int32 UserIndex, bool bSuccess);
 
 	/**
