Fix assertion when attempting to load enhanced input user settings

#rnx

diff --git a/Engine/Plugins/EnhancedInput/Source/EnhancedInput/Private/EnhancedInputSubsystems.cpp b/Engine/Plugins/EnhancedInput/Source/EnhancedInput/Private/EnhancedInputSubsystems.cpp
--- a/Engine/Plugins/EnhancedInput/Source/EnhancedInput/Private/EnhancedInputSubsystems.cpp
+++ b/Engine/Plugins/EnhancedInput/Source/EnhancedInput/Private/EnhancedInputSubsystems.cpp
@@ -104,6 +104,13 @@ void UEnhancedInputLocalPlayerSubsystem::InitalizeUserSettings()
 		UE_LOG(LogEnhancedInput, Verbose, TEXT("bEnableUserSettings is set to false, skipping creation of UEnhancedInputUserSettings!"));
 		return;
 	}
+
+	// We can't load the user settings without a valid player input object because we get the save file name from it
+	if (!GetPlayerInput())
+	{
+		UE_LOG(LogEnhancedInput, Warning, TEXT("UEnhancedInputLocalPlayerSubsystem for Local Player '%s' does not have a valid PlayerInput object. Failed to load user settings."), *GetNameSafe(LP));
+		return;
+	}
 	
 	UserSettings = UEnhancedInputUserSettings::LoadOrCreateSettings(LP);
 
