EnhancedInput: Add RegistrationTrackingMode to Input Mapping Context to select how multiple registrations of the IMC is handled

- The default behaviour is "Untracked"
- New option "CountRegistrations" will track how many calls to AddMappingContext() are made for the IMC and keep the IMC applied until the same number of calls to RemoveMappingContext() are made

#rnx

diff --git a/Engine/Plugins/EnhancedInput/Source/EnhancedInput/Private/EnhancedInputSubsystemInterface.cpp b/Engine/Plugins/EnhancedInput/Source/EnhancedInput/Private/EnhancedInputSubsystemInterface.cpp
--- a/Engine/Plugins/EnhancedInput/Source/EnhancedInput/Private/EnhancedInputSubsystemInterface.cpp
+++ b/Engine/Plugins/EnhancedInput/Source/EnhancedInput/Private/EnhancedInputSubsystemInterface.cpp
@@ -224,7 +224,7 @@ void IEnhancedInputSubsystemInterface::ClearAllMappings()
 {
 	if (UEnhancedPlayerInput* PlayerInput = GetPlayerInput())
 	{
-		PlayerInput->AppliedInputContexts.Empty();
+		PlayerInput->AppliedInputContextData.Empty();
 		RequestRebuildControlMappings();
 	}
 }
@@ -236,8 +236,29 @@ void IEnhancedInputSubsystemInterface::AddMappingContext(const UInputMappingCont
 	{
 		if (UEnhancedPlayerInput* PlayerInput = GetPlayerInput())
 		{
-			PlayerInput->AppliedInputContexts.Add(MappingContext, Priority);
-			RequestRebuildControlMappings(Options);
+			const EMappingContextRegistrationTrackingMode TrackingMode = MappingContext->GetRegistrationTrackingMode();
+			if (TrackingMode == EMappingContextRegistrationTrackingMode::Untracked)
+			{
+				PlayerInput->AppliedInputContextData.Add(MappingContext, { Priority });
+				RequestRebuildControlMappings(Options);
+			}
+			else if (TrackingMode == EMappingContextRegistrationTrackingMode::CountRegistrations)
+			{
+				if (FAppliedInputContextData* IMCDataPtr = PlayerInput->AppliedInputContextData.Find(MappingContext))
+				{
+					IMCDataPtr->RegistrationCount++;
+				}
+				else
+				{
+					constexpr int32 InitialRegistrationCount = 1;
+					PlayerInput->AppliedInputContextData.Add(MappingContext, { Priority, InitialRegistrationCount });
+					RequestRebuildControlMappings(Options);
+				}
+			}
+			else
+			{
+				checkNoEntry(); // Unhandled tracking mode
+			}
 		}
 
 		if (Options.bNotifyUserSettings)
@@ -258,13 +279,37 @@ void IEnhancedInputSubsystemInterface::RemoveMappingContext(const UInputMappingC
 {
 	if (MappingContext)
 	{
+		bool bDidRemoveMappingContext = false;
 		if (UEnhancedPlayerInput* PlayerInput = GetPlayerInput())
 		{
-			PlayerInput->AppliedInputContexts.Remove(MappingContext);
-			RequestRebuildControlMappings(Options);
+			const EMappingContextRegistrationTrackingMode TrackingMode = MappingContext->GetRegistrationTrackingMode();
+			if (TrackingMode == EMappingContextRegistrationTrackingMode::Untracked)
+			{
+				PlayerInput->AppliedInputContextData.Remove(MappingContext);
+				RequestRebuildControlMappings(Options);
+				bDidRemoveMappingContext = true;
+			}
+			else if (TrackingMode == EMappingContextRegistrationTrackingMode::CountRegistrations)
+			{
+				if (FAppliedInputContextData* IMCDataPtr = PlayerInput->AppliedInputContextData.Find(MappingContext))
+				{
+					const int32 RegistrationCount = --(IMCDataPtr->RegistrationCount);
+					ensureMsgf(RegistrationCount >= 0, TEXT("Input Mapping Context [%s] has a negative registration count without being removed"), *MappingContext->GetName());
+					if (RegistrationCount <= 0)
+					{
+						PlayerInput->AppliedInputContextData.Remove(MappingContext);
+						RequestRebuildControlMappings(Options);
+						bDidRemoveMappingContext = true;
+					}
+				}
+			}
+			else
+			{
+				checkNoEntry(); // Unhandled tracking mode
+			}
 		}
 
-		if (Options.bNotifyUserSettings)
+		if (bDidRemoveMappingContext && Options.bNotifyUserSettings)
 		{
 			if (UEnhancedInputUserSettings* Settings = GetUserSettings())
 			{
@@ -375,12 +420,12 @@ EMappingQueryResult IEnhancedInputSubsystemInterface::QueryMapKeyInActiveContext
 	}
 
 	// TODO: Inefficient, but somewhat forgivable as the mapping context count is likely to be single figure.
-	TMap<TObjectPtr<const UInputMappingContext>, int32> OrderedInputContexts = PlayerInput->AppliedInputContexts;
-	OrderedInputContexts.ValueSort([](const int32& A, const int32& B) { return A > B; });
+	TMap<TObjectPtr<const UInputMappingContext>, FAppliedInputContextData> OrderedInputContexts = PlayerInput->AppliedInputContextData;
+	OrderedInputContexts.ValueSort([](const FAppliedInputContextData& A, const FAppliedInputContextData& B) { return A.Priority > B.Priority; });
 
 	TArray<UInputMappingContext*> Applied;
 	Applied.Reserve(OrderedInputContexts.Num());
-	for (const TPair<TObjectPtr<const UInputMappingContext>, int32>& ContextPair : OrderedInputContexts)
+	for (const TPair<TObjectPtr<const UInputMappingContext>, FAppliedInputContextData>& ContextPair : OrderedInputContexts)
 	{
 		Applied.Add(const_cast<UInputMappingContext*>(ToRawPtr(ContextPair.Key)));
 	}
@@ -625,9 +670,9 @@ bool IEnhancedInputSubsystemInterface::HasMappingContext(const UInputMappingCont
 	
 	if (const UEnhancedPlayerInput* const Input = GetPlayerInput())
 	{
-		if (const int32* const FoundPriority = Input->AppliedInputContexts.Find(MappingContext))
+		if (const FAppliedInputContextData* const FoundInputContextData = Input->AppliedInputContextData.Find(MappingContext))
 		{
-			OutFoundPriority = *FoundPriority;
+			OutFoundPriority = FoundInputContextData->Priority;
 			bResult = true;
 		}
 	}
@@ -674,6 +719,21 @@ TArray<FEnhancedActionKeyMapping> IEnhancedInputSubsystemInterface::GetAllPlayer
 	return PlayerMappableMappings;
 }
 
+void IEnhancedInputSubsystemInterface::ValidateTrackedMappingContextsAreUnregistered() const
+{
+	if (UEnhancedPlayerInput* PlayerInput = GetPlayerInput())
+	{
+		for (const TPair<TObjectPtr<const UInputMappingContext>, FAppliedInputContextData>& Pair : PlayerInput->GetAppliedInputContextData())
+		{
+			if (IsValid(Pair.Key) && Pair.Key->GetRegistrationTrackingMode() != EMappingContextRegistrationTrackingMode::Untracked)
+			{
+				UE_LOG(LogEnhancedInput, Warning, TEXT("Input Mapping Context [%s] has tracking mode [%s] but is still applied and might be leaking, unregister it before deinitialization."),
+					*Pair.Key->GetName(), *UEnum::GetValueAsString(Pair.Key->GetRegistrationTrackingMode()));
+			}
+		}
+	}
+}
+
 // TODO: This should be a delegate (along with InjectChordBlockers), moving chording out of the underlying subsystem and enabling implementation of custom mapping handlers.
 /**
  * Reorder the given UnordedMappings such that chording mappings > chorded mappings > everything else.
@@ -815,7 +875,7 @@ void IEnhancedInputSubsystemInterface::RebuildControlMappings()
 	AppliedContextRedirects.Reset();
 
 	// Order contexts by priority
-	TMap<TObjectPtr<const UInputMappingContext>, int32> OrderedInputContexts = PlayerInput->AppliedInputContexts;
+	TMap<TObjectPtr<const UInputMappingContext>, FAppliedInputContextData> OrderedInputContexts = PlayerInput->AppliedInputContextData;
 	
 	// Replace any mapping contexts that may have a redirect on this platform
 	if (UEnhancedInputPlatformSettings* PlatformSettings = UEnhancedInputPlatformSettings::Get())
@@ -831,10 +891,10 @@ void IEnhancedInputSubsystemInterface::RebuildControlMappings()
 			}
 			
 			// Replace the existing IMC with the one that it should be redirected to on the PlayerInput 
-			if (const int32* ExistingIMCPriority = OrderedInputContexts.Find(Pair.Key))
+			if (const FAppliedInputContextData* ExistingIMCData = OrderedInputContexts.Find(Pair.Key))
 			{
 				OrderedInputContexts.Remove(Pair.Key);
-				OrderedInputContexts.Add(Pair.Value, *ExistingIMCPriority);
+				OrderedInputContexts.Add(Pair.Value, *ExistingIMCData);
 				AppliedContextRedirects.Add(Pair);
 
 				// Optional logging that may be helpful for debugging purposes
@@ -852,7 +912,7 @@ void IEnhancedInputSubsystemInterface::RebuildControlMappings()
 	}
 	
 	// Order contexts by priority
-	OrderedInputContexts.ValueSort([](const int32& A, const int32& B) { return A > B; });
+	OrderedInputContexts.ValueSort([](const FAppliedInputContextData& A, const FAppliedInputContextData& B) { return A.Priority > B.Priority; });
 
 	TSet<FKey> AppliedKeys;
 
@@ -869,7 +929,7 @@ void IEnhancedInputSubsystemInterface::RebuildControlMappings()
 	// an input mapping context.
 	TArray<FKey> MappedKeysToActionName;
 	
-	for (const TPair<TObjectPtr<const UInputMappingContext>, int32>& ContextPair : OrderedInputContexts)
+	for (const TPair<TObjectPtr<const UInputMappingContext>, FAppliedInputContextData>& ContextPair : OrderedInputContexts)
 	{
 		// Don't apply context specific keys immediately, allowing multiple mappings to the same key within the same context if required.
 		TArray<FKey> ContextAppliedKeys;
diff --git a/Engine/Plugins/EnhancedInput/Source/EnhancedInput/Private/EnhancedInputSubsystemInterfaceDebug.cpp b/Engine/Plugins/EnhancedInput/Source/EnhancedInput/Private/EnhancedInputSubsystemInterfaceDebug.cpp
--- a/Engine/Plugins/EnhancedInput/Source/EnhancedInput/Private/EnhancedInputSubsystemInterfaceDebug.cpp
+++ b/Engine/Plugins/EnhancedInput/Source/EnhancedInput/Private/EnhancedInputSubsystemInterfaceDebug.cpp
@@ -102,11 +102,11 @@ void IEnhancedInputSubsystemInterface::ShowMappingContextDebugInfo(UCanvas* Canv
 			}
 		};
 
-		TMap<TObjectPtr<const UInputMappingContext>, int32> OrderedInputContexts = PlayerInput->AppliedInputContexts;
-		OrderedInputContexts.ValueSort([](const int32& A, const int32& B) { return A > B; });
+		TMap<TObjectPtr<const UInputMappingContext>, FAppliedInputContextData> OrderedInputContexts = PlayerInput->AppliedInputContextData;
+		OrderedInputContexts.ValueSort([](const FAppliedInputContextData& A, const FAppliedInputContextData& B) { return A.Priority > B.Priority; });
 		
 		// Work through all input contexts, displaying active mappings, overridden mappings, etc.
-		for (const TPair<TObjectPtr<const UInputMappingContext>, int32>& ContextPair : OrderedInputContexts)
+		for (const TPair<TObjectPtr<const UInputMappingContext>, FAppliedInputContextData>& ContextPair : OrderedInputContexts)
 		{
 			const UInputMappingContext* AppliedContext = ContextPair.Key.Get();
 
@@ -121,19 +121,43 @@ void IEnhancedInputSubsystemInterface::ShowMappingContextDebugInfo(UCanvas* Canv
 			}
 
 			DisplayDebugManager.SetDrawColor(FColor::Yellow);
+
+			const EMappingContextRegistrationTrackingMode TrackingMode = AppliedContext->GetRegistrationTrackingMode();
+			FString TrackingModeDebugString;
+			switch (TrackingMode)
+			{
+			case EMappingContextRegistrationTrackingMode::Untracked:
+				// Nothing to show for untracked
+				break;
+			case EMappingContextRegistrationTrackingMode::CountRegistrations:
+				TrackingModeDebugString = FString::Printf(TEXT(" | Registration Count: %d"), ContextPair.Value.RegistrationCount);
+				break;
+			default:
+				checkNoEntry(); // Unhandled tracking mode
+				break;
+			}
 			
 			// If this context was redirected via platform settings, then add some debug info about it here
 			if (const TObjectPtr<const UInputMappingContext>* RedirectedContextPtr = AppliedContextRedirects.Find(AppliedContext))
 			{
 				const UInputMappingContext* RedirectedContext = *RedirectedContextPtr;
-				DisplayDebugManager.DrawString(FString::Printf(TEXT("  Redirected Context: %s -> %s | Priority: %d"), *AppliedContext->GetFName().ToString(), *RedirectedContext->GetFName().ToString(), ContextPair.Value));
+				DisplayDebugManager.DrawString(FString::Printf(TEXT("  Redirected Context: %s -> %s | Priority: %d | Tracking Mode: %s%s"),
+					*AppliedContext->GetFName().ToString(),
+					*RedirectedContext->GetFName().ToString(),
+					ContextPair.Value.Priority,
+					*UEnum::GetValueAsString(TrackingMode),
+					*TrackingModeDebugString));
 
 				// Change the current context that is being used to display the correct mappings
 				AppliedContext = RedirectedContext;
 			}
 			else
 			{
-				DisplayDebugManager.DrawString(FString::Printf(TEXT("  Context: %s | Priority: %d"), *AppliedContext->GetFName().ToString(), ContextPair.Value));
+				DisplayDebugManager.DrawString(FString::Printf(TEXT("  Context: %s | Priority: %d | Tracking Mode: %s%s"),
+					*AppliedContext->GetFName().ToString(),
+					ContextPair.Value.Priority,
+					*UEnum::GetValueAsString(TrackingMode),
+					*TrackingModeDebugString));
 			}
 
 			// Build a table of mappings per action
diff --git a/Engine/Plugins/EnhancedInput/Source/EnhancedInput/Private/EnhancedInputSubsystems.cpp b/Engine/Plugins/EnhancedInput/Source/EnhancedInput/Private/EnhancedInputSubsystems.cpp
--- a/Engine/Plugins/EnhancedInput/Source/EnhancedInput/Private/EnhancedInputSubsystems.cpp
+++ b/Engine/Plugins/EnhancedInput/Source/EnhancedInput/Private/EnhancedInputSubsystems.cpp
@@ -33,6 +33,13 @@ namespace UE::Input
 		TEXT("Should the UEnhancedInputDeveloperSettings::DefaultMappingContexts be applied to every UEnhancedPlayerInput?"));
 }
 
+void UEnhancedInputLocalPlayerSubsystem::Deinitialize()
+{
+	Super::Deinitialize();
+
+	ValidateTrackedMappingContextsAreUnregistered();
+}
+
 void UEnhancedInputLocalPlayerSubsystem::PlayerControllerChanged(APlayerController* NewPlayerController)
 {
 	Super::PlayerControllerChanged(NewPlayerController);
@@ -201,6 +208,8 @@ void UEnhancedInputWorldSubsystem::Deinitialize()
 		FSlateApplication::Get().UnregisterInputPreProcessor(InputPreprocessor);
 	}
 	
+	ValidateTrackedMappingContextsAreUnregistered();
+
 	InputPreprocessor.Reset();
 	CurrentInputStack.Empty();
 	PlayerInput = nullptr;
diff --git a/Engine/Plugins/EnhancedInput/Source/EnhancedInput/Public/EnhancedInputSubsystemInterface.h b/Engine/Plugins/EnhancedInput/Source/EnhancedInput/Public/EnhancedInputSubsystemInterface.h
--- a/Engine/Plugins/EnhancedInput/Source/EnhancedInput/Public/EnhancedInputSubsystemInterface.h
+++ b/Engine/Plugins/EnhancedInput/Source/EnhancedInput/Public/EnhancedInputSubsystemInterface.h
@@ -377,6 +377,9 @@ public:
 	 */
 	UFUNCTION(BlueprintCallable, BlueprintCosmetic, Category = "Input|Mapping Queries")
 	virtual TArray<FEnhancedActionKeyMapping> GetAllPlayerMappableActionKeyMappings() const;
+
+	/** Logs warnings for any currently applied Input Mapping Contexts that have a tracked RegistrationTrackingMode to detect leaks */
+	void ValidateTrackedMappingContextsAreUnregistered() const;
 	
 protected:
 
diff --git a/Engine/Plugins/EnhancedInput/Source/EnhancedInput/Public/EnhancedInputSubsystems.h b/Engine/Plugins/EnhancedInput/Source/EnhancedInput/Public/EnhancedInputSubsystems.h
--- a/Engine/Plugins/EnhancedInput/Source/EnhancedInput/Public/EnhancedInputSubsystems.h
+++ b/Engine/Plugins/EnhancedInput/Source/EnhancedInput/Public/EnhancedInputSubsystems.h
@@ -23,6 +23,7 @@ class ENHANCEDINPUT_API UEnhancedInputLocalPlayerSubsystem : public ULocalPlayer
 public:
 
 	// Begin ULocalPlayerSubsystem
+	virtual void Deinitialize() override;
 	virtual void PlayerControllerChanged(APlayerController* NewPlayerController) override;
 	// End ULocalPlayerSubsystem
 	
diff --git a/Engine/Plugins/EnhancedInput/Source/EnhancedInput/Public/EnhancedPlayerInput.h b/Engine/Plugins/EnhancedInput/Source/EnhancedInput/Public/EnhancedPlayerInput.h
--- a/Engine/Plugins/EnhancedInput/Source/EnhancedInput/Public/EnhancedPlayerInput.h
+++ b/Engine/Plugins/EnhancedInput/Source/EnhancedInput/Public/EnhancedPlayerInput.h
@@ -69,6 +69,22 @@ struct FInjectedInputArray
 	TArray<FInjectedInput> Injected;
 };
 
+USTRUCT()
+struct FAppliedInputContextData
+{
+	GENERATED_BODY()
+
+	UPROPERTY()
+	int32 Priority = 0;
+
+	/**
+	 * Used when RegistrationTrackingMode is set to CountRegistrations
+	 * This is how many times the IMC has been added minus how many times it has been removed. The IMC is unregistered when this hits 0.
+	 */
+	UPROPERTY()
+	int32 RegistrationCount = 0;
+};
+
 /**
 * UEnhancedPlayerInput : UPlayerInput extensions for enhanced player input system
 */
@@ -143,8 +159,15 @@ protected:
 	// Causes key to be consumed if it is affecting an action.
 	virtual bool IsKeyHandledByAction(FKey Key) const override;
 	
-	/** Note: Source reference only. Use EnhancedActionMappings for the actual mappings (with properly instanced triggers/modifiers) */
-	const TMap<TObjectPtr<const UInputMappingContext>, int32>& GetAppliedInputContexts() const { return AppliedInputContexts; }
+	/** Note: Source reference only. Use GetEnhancedActionMappings() for the actual mappings (with properly instanced triggers/modifiers) */
+	const TMap<TObjectPtr<const UInputMappingContext>, FAppliedInputContextData>& GetAppliedInputContextData() const { return AppliedInputContextData; }
+	UE_DEPRECATED(5.6, "GetAppliedInputContexts() is deprecated, use GetAppliedInputContextData() instead")
+	const TMap<TObjectPtr<const UInputMappingContext>, int32>& GetAppliedInputContexts() const
+	{
+		PRAGMA_DISABLE_DEPRECATION_WARNINGS
+		return AppliedInputContexts;
+		PRAGMA_ENABLE_DEPRECATION_WARNINGS
+	}
 
 	/** This player's version of the Action Mappings */
 	const TArray<FEnhancedActionKeyMapping>& GetEnhancedActionMappings() const { return EnhancedActionMappings; }
@@ -200,7 +223,10 @@ private:
 	 *
 	 * These mapping contexts will only have their mappings processed if the current input mode
 	 * matches the query set on them.
-	 */
+	*/
+	UPROPERTY(Transient)
+	TMap<TObjectPtr<const UInputMappingContext>, FAppliedInputContextData> AppliedInputContextData;
+	UE_DEPRECATED(5.6, "AppliedInputContexts is deprecated, use AppliedInputContextData instead")
 	UPROPERTY(Transient)
 	TMap<TObjectPtr<const UInputMappingContext>, int32> AppliedInputContexts;
 
diff --git a/Engine/Plugins/EnhancedInput/Source/EnhancedInput/Public/InputMappingContext.h b/Engine/Plugins/EnhancedInput/Source/EnhancedInput/Public/InputMappingContext.h
--- a/Engine/Plugins/EnhancedInput/Source/EnhancedInput/Public/InputMappingContext.h
+++ b/Engine/Plugins/EnhancedInput/Source/EnhancedInput/Public/InputMappingContext.h
@@ -38,6 +38,29 @@ enum class EMappingContextInputModeFilterOptions : uint8
 	DoNotFilter
 };
 
+/**
+ * Options for how multiple registrations of an input mapping context should be tracked.
+ */
+UENUM()
+enum class EMappingContextRegistrationTrackingMode : uint8
+{
+	/**
+	 * This is the default behavior.
+	 * Registrations of the Input Mapping Context are not tracked. The mapping context will be unregistered when removing it the first time, no matter how many times it has been added.
+	 */
+	Untracked,
+
+	/**
+	 * Track how many times the IMC is added and keeps the IMC applied until the IMC is removed the same number of times.
+	 * This allows multiple systems to use the same Input Mapping Context without needing to check if any other systems are still using the same Input Mapping Context.
+	 *
+	 * Warnings will be logged if Input Mapping Contexts with this tracking mode are still applied at deinitialization, the expectation is that there will be a RemoveMappingContext() call for every call to AddMappingContext() when using this mode.
+	 * 
+	 * @see IEnhancedInputSubsystemInterface::ValidateTrackedMappingContextsAreUnregistered
+	 */
+	CountRegistrations
+};
+
 /**
 * UInputMappingContext : A collection of key to action mappings for a specific input context
 * Could be used to:
@@ -85,6 +108,12 @@ protected:
 	UPROPERTY(config, BlueprintReadOnly, EditAnywhere, Category = "Input Modes", meta = (EditConditionHides, EditCondition="ShouldShowInputModeQuery() && InputModeFilterOptions == EMappingContextInputModeFilterOptions::UseCustomQuery"))
 	FGameplayTagQuery InputModeQueryOverride;
 
+	/**
+	 * Select the behaviour when multiple AddMappingContext() calls are made for this Input Mapping Context
+	 */
+	UPROPERTY(BlueprintReadOnly, EditAnywhere, Category = "Registration")
+	EMappingContextRegistrationTrackingMode RegistrationTrackingMode = EMappingContextRegistrationTrackingMode::Untracked;
+
 public:
 
 	/**
@@ -97,6 +126,11 @@ public:
 	 * out based on the current input mode or not. 
 	 */
 	FGameplayTagQuery GetInputModeQuery() const;
+
+	/**
+	 * @return The registration tracking mode that this IMC is using. 
+	 */
+	EMappingContextRegistrationTrackingMode GetRegistrationTrackingMode() const { return RegistrationTrackingMode; }
 	
 	// Localized context descriptor
 	UPROPERTY(BlueprintReadOnly, EditAnywhere, Category = "Description", DisplayName = "Description")
diff --git a/Engine/Plugins/EnhancedInput/Source/InputEditor/Private/Tests/InputSystemTest.cpp b/Engine/Plugins/EnhancedInput/Source/InputEditor/Private/Tests/InputSystemTest.cpp
--- a/Engine/Plugins/EnhancedInput/Source/InputEditor/Private/Tests/InputSystemTest.cpp
+++ b/Engine/Plugins/EnhancedInput/Source/InputEditor/Private/Tests/InputSystemTest.cpp
@@ -192,6 +192,34 @@ bool FInputInjectionTest::RunTest(const FString& Parameters)
 	return true;
 }
 
+IMPLEMENT_SIMPLE_AUTOMATION_TEST(FInputContextRegistrationTrackingTest, "Input.System.ContextRegistrationTracking", BasicSystemTestFlags)
+bool FInputContextRegistrationTrackingTest::RunTest(const FString& Parameters)
+{
+	UWorld* World = AnEmptyWorld();
+	UControllablePlayer& Data =
+	GIVEN(AControllablePlayer(World));
+	AND(const FName UntrackedMappingName = TEXT("UntrackedMapping"));
+	AND(const FName CountedMappingName = TEXT("CountedMapping"));
+
+	// Test 1 - An untracked mapping that is added twice then removed once is no longer applied
+	WHEN(AnInputContextWithTrackingModeIsAppliedToAPlayer(Data, UntrackedMappingName, 0, EMappingContextRegistrationTrackingMode::Untracked));
+	AND(AnInputContextIsReappliedToAPlayer(Data, UntrackedMappingName, 0));
+	AND(AnInputContextIsRemovedFromAPlayer(Data, UntrackedMappingName));
+	THEN(InputContextIsNotApplied(Data, UntrackedMappingName));
+
+	// Test 2a - A registration counted mapping is added twice and then removed once is still applied
+	WHEN(AnInputContextWithTrackingModeIsAppliedToAPlayer(Data, CountedMappingName, 0, EMappingContextRegistrationTrackingMode::CountRegistrations));
+	AND(AnInputContextIsReappliedToAPlayer(Data, CountedMappingName, 0));
+	AND(AnInputContextIsRemovedFromAPlayer(Data, CountedMappingName));
+	THEN(InputContextIsApplied(Data, CountedMappingName));
+
+	// Test 2b - The registration counted mapping is removed once more and then no longer applied
+	WHEN(AnInputContextIsRemovedFromAPlayer(Data, CountedMappingName));
+	THEN(InputContextIsNotApplied(Data, CountedMappingName));
+
+	return true;
+}
+
 // Event transition tests
 
 IMPLEMENT_SIMPLE_AUTOMATION_TEST(FInputEventTransitionTest, "Input.System.Events.Transitions", BasicSystemTestFlags)	// TODO: Split into individual transition tests?
diff --git a/Engine/Plugins/EnhancedInput/Source/InputEditor/Private/Tests/InputTestFramework.cpp b/Engine/Plugins/EnhancedInput/Source/InputEditor/Private/Tests/InputTestFramework.cpp
--- a/Engine/Plugins/EnhancedInput/Source/InputEditor/Private/Tests/InputTestFramework.cpp
+++ b/Engine/Plugins/EnhancedInput/Source/InputEditor/Private/Tests/InputTestFramework.cpp
@@ -83,6 +83,45 @@ UInputMappingContext* AnInputContextIsAppliedToAPlayer(UControllablePlayer& Play
 	return Context;
 }
 
+UInputMappingContext* AnInputContextWithTrackingModeIsAppliedToAPlayer(UControllablePlayer& PlayerData, FName ContextName, int32 WithPriority, EMappingContextRegistrationTrackingMode TrackingMode)
+{
+	// IMC Extension that allows simulating a user setting a tracking mode in the editor
+	class UInputMappingContextWithTrackingSetter : public UInputMappingContext
+	{
+	public:
+		void SetTrackingMode(EMappingContextRegistrationTrackingMode InTrackingMode)
+		{
+			RegistrationTrackingMode = InTrackingMode;
+		}
+	};
+	
+	UInputMappingContextWithTrackingSetter* Context = NewObject<UInputMappingContextWithTrackingSetter>(PlayerData.Player, ContextName);
+	Context->SetTrackingMode(TrackingMode);
+	PlayerData.InputContext.Emplace(ContextName, Context);
+	PlayerData.Subsystem->AddMappingContext(Context, WithPriority);
+	return Context;
+}
+
+UInputMappingContext* AnInputContextIsReappliedToAPlayer(UControllablePlayer& PlayerData, FName ContextName, int32 WithPriority)
+{
+	UInputMappingContext* Context = PlayerData.InputContext.FindRef(ContextName);
+	if (ensureMsgf(Context, TEXT("For an IMC to be reapplied, it must already exist")))
+	{
+		PlayerData.Subsystem->AddMappingContext(Context, WithPriority);
+	}
+	return Context;
+}
+
+UInputMappingContext* AnInputContextIsRemovedFromAPlayer(UControllablePlayer& PlayerData, FName ContextName)
+{
+	UInputMappingContext* Context = PlayerData.InputContext.FindRef(ContextName);
+	if (ensureMsgf(Context, TEXT("For an IMC to be removed, it must already exist")))
+	{
+		PlayerData.Subsystem->RemoveMappingContext(Context);
+	}
+	return Context;
+}
+
 UInputAction* AnInputAction(UControllablePlayer& PlayerData, FName ActionName, EInputActionValueType ValueType)
 {
 	UTestMappableKeysAction* Action = NewObject<UTestMappableKeysAction>(PlayerData.Player, ActionName);
@@ -314,3 +353,8 @@ const FInputActionInstance& FInputTestHelper::GetActionData(UControllablePlayer&
 	static FInputActionInstance NoActionData;
 	return NoActionData;
 }
+
+bool FInputTestHelper::HasAppliedInputMappingContext(UControllablePlayer& Data, FName ContextName)
+{
+	return Data.PlayerInput->GetAppliedInputContextData().Contains(FInputTestHelper::FindContext(Data, ContextName));
+}
diff --git a/Engine/Plugins/EnhancedInput/Source/InputEditor/Private/Tests/InputTestFramework.h b/Engine/Plugins/EnhancedInput/Source/InputEditor/Private/Tests/InputTestFramework.h
--- a/Engine/Plugins/EnhancedInput/Source/InputEditor/Private/Tests/InputTestFramework.h
+++ b/Engine/Plugins/EnhancedInput/Source/InputEditor/Private/Tests/InputTestFramework.h
@@ -7,6 +7,7 @@
 #include "UserSettings/EnhancedInputUserSettings.h"
 #include "InputTestFramework.generated.h"
 
+enum class EMappingContextRegistrationTrackingMode : uint8;
 struct FEnhancedActionKeyMapping;
 class UEnhancedPlayerInput;
 
@@ -246,6 +247,8 @@ struct FInputTestHelper
 	static void ResetActionInstanceData(UControllablePlayer& Data);
 
 	static const struct FInputActionInstance& GetActionData(UControllablePlayer& Data, FName ActionName);
+
+	static bool HasAppliedInputMappingContext(UControllablePlayer& Data, FName ContextName);
 };
 
 // TODO: Switch to something like Catch2?
@@ -261,6 +264,9 @@ struct FInputTestHelper
 UWorld* AnEmptyWorld();
 UControllablePlayer& AControllablePlayer(UWorld* World);
 UInputMappingContext* AnInputContextIsAppliedToAPlayer(UControllablePlayer& PlayerData, FName ContextName, int32 WithPriority);
+UInputMappingContext* AnInputContextWithTrackingModeIsAppliedToAPlayer(UControllablePlayer& PlayerData, FName ContextName, int32 WithPriority, EMappingContextRegistrationTrackingMode TrackingMode);
+UInputMappingContext* AnInputContextIsReappliedToAPlayer(UControllablePlayer& PlayerData, FName ContextName, int32 WithPriority);
+UInputMappingContext* AnInputContextIsRemovedFromAPlayer(UControllablePlayer& PlayerData, FName ContextName);
 UInputAction* AnInputAction(UControllablePlayer& PlayerData, FName ActionName, EInputActionValueType ValueType);
 FEnhancedActionKeyMapping& AnActionIsMappedToAKey(UControllablePlayer& PlayerData, FName ContextName, FName ActionName, FKey Key); // NOTE: Adds a template mapping to the input context but returns the resulting live action mapping on the player
 
@@ -296,3 +302,5 @@ inline bool ReleasingKeyTriggersCanceled(UControllablePlayer& Data, FName Action
 inline bool ReleasingKeyTriggersCompleted(UControllablePlayer& Data, FName ActionName) { return FInputTestHelper::TestCompleted(Data, ActionName); }
 inline bool ReleasingKeyTriggersAction(UControllablePlayer& Data, FName ActionName) { return FInputTestHelper::TestTriggered(Data, ActionName); }
 
+inline bool InputContextIsApplied(UControllablePlayer& Data, FName ContextName) { return FInputTestHelper::HasAppliedInputMappingContext(Data, ContextName); }
+inline bool InputContextIsNotApplied(UControllablePlayer& Data, FName ContextName) { return !FInputTestHelper::HasAppliedInputMappingContext(Data, ContextName); }
