EnhancedInput: Fix validation error when saving widget blueprints with disconnected input action event nodes

The step that sets `bAutomaticallyRegisterInputOnConstruction` to true on widget blueprints with input action event nodes is skipped when no event pins are active so we now do not consider those inactive nodes when validating widget blueprints.

#rnx

diff --git a/Engine/Plugins/EnhancedInput/Source/InputBlueprintNodes/Private/EnhancedInputUserWidgetValidator.cpp b/Engine/Plugins/EnhancedInput/Source/InputBlueprintNodes/Private/EnhancedInputUserWidgetValidator.cpp
--- a/Engine/Plugins/EnhancedInput/Source/InputBlueprintNodes/Private/EnhancedInputUserWidgetValidator.cpp
+++ b/Engine/Plugins/EnhancedInput/Source/InputBlueprintNodes/Private/EnhancedInputUserWidgetValidator.cpp
@@ -17,14 +17,14 @@ namespace UE::Input
 		ECVF_Default);
 
 	/**
-	 * Returns true if the given widget BP has any nodes of the enhanced input type
+	 * Returns true if the given widget BP has any linked nodes of the enhanced input type
 	 * in any of its graphs.
 	 */
-	static bool HasAnyEnhancedInputNodes(const UWidgetBlueprint* WidgetBP)
+	static bool HasAnyActiveEnhancedInputNodes(const UWidgetBlueprint* WidgetBP)
 	{
 		check(WidgetBP);
 		
-		bool bHasAnyEnhancedInputNodes = false;
+		bool bHasAnyActiveEnhancedInputNodes = false;
 		
 		TArray<UEdGraph*> Graphs;
 		WidgetBP->GetAllGraphs(Graphs);
@@ -34,15 +34,22 @@ namespace UE::Input
 			TArray<UK2Node_EnhancedInputAction*> EventNodes;
 			Graph->GetNodesOfClass(EventNodes);
 
-			bHasAnyEnhancedInputNodes |= !EventNodes.IsEmpty();
+			for (UK2Node_EnhancedInputAction* EventNode : EventNodes)
+			{
+				if (EventNode && EventNode->HasAnyConnectedEventPins())
+				{
+					bHasAnyActiveEnhancedInputNodes = true;
+					break;
+				}
+			}
 
-			if (bHasAnyEnhancedInputNodes)
+			if (bHasAnyActiveEnhancedInputNodes)
 			{
 				break;
 			}
 		}
 
-		return bHasAnyEnhancedInputNodes;
+		return bHasAnyActiveEnhancedInputNodes;
 	}
 }
 
@@ -71,7 +78,7 @@ EDataValidationResult UEnhancedInputUserWidgetValidator::ValidateWidgetBlueprint
 	}
 
 	// If there are no Enhanced Input nodes then there is no need to check further
-	if (!UE::Input::HasAnyEnhancedInputNodes(WidgetBP))
+	if (!UE::Input::HasAnyActiveEnhancedInputNodes(WidgetBP))
 	{
 		return EDataValidationResult::Valid;
 	}
@@ -84,7 +91,7 @@ EDataValidationResult UEnhancedInputUserWidgetValidator::ValidateWidgetBlueprint
 
 	// ... otherwise, this widget is invalid.
 	// Widgets require the bAutomaticallyRegisterInputOnConstruction to be true in order to receive enhanced input events. 
-	const FText CurrentError = NSLOCTEXT("EnhancedInput", "Input.Widget.Error", "'bAutomaticallyRegisterInputOnConstruction' must be true in order to use Enhanced Input in the widget");
+	const FText CurrentError = NSLOCTEXT("EnhancedInput", "Input.Widget.Error", "'bAutomaticallyRegisterInputOnConstruction' failed to automatically update but must be true in order to use Enhanced Input in the widget");
 	
 	AssetMessage(InAssetData, EMessageSeverity::Error, CurrentError);
 	return EDataValidationResult::Invalid;	
diff --git a/Engine/Plugins/EnhancedInput/Source/InputBlueprintNodes/Private/K2Node_EnhancedInputAction.cpp b/Engine/Plugins/EnhancedInput/Source/InputBlueprintNodes/Private/K2Node_EnhancedInputAction.cpp
--- a/Engine/Plugins/EnhancedInput/Source/InputBlueprintNodes/Private/K2Node_EnhancedInputAction.cpp
+++ b/Engine/Plugins/EnhancedInput/Source/InputBlueprintNodes/Private/K2Node_EnhancedInputAction.cpp
@@ -49,14 +49,17 @@ UK2Node_EnhancedInputAction::UK2Node_EnhancedInputAction(const FObjectInitialize
 {
 }
 
-void ForEachEventPinName(TFunctionRef<void(ETriggerEvent Event, FName PinName)> PinLambda)
+void ForEachEventPinName(TFunctionRef<bool(ETriggerEvent Event, FName PinName)> PinLambda)
 {
 	UEnum* EventEnum = StaticEnum<ETriggerEvent>();
 	for (int32 i = 0; i < EventEnum->NumEnums() - 1; ++i)
 	{
 		if (!EventEnum->HasMetaData(TEXT("Hidden"), i))
 		{
-			PinLambda(ETriggerEvent(EventEnum->GetValueByIndex(i)), *EventEnum->GetNameStringByIndex(i));
+			if (!PinLambda(ETriggerEvent(EventEnum->GetValueByIndex(i)), *EventEnum->GetNameStringByIndex(i)))
+			{
+				break;
+			}
 		}
 	}
 }
@@ -141,6 +144,9 @@ void UK2Node_EnhancedInputAction::AllocateDefaultPins()
 			NewPin->PinToolTip += UnsuportedTooltip.ToString();
 			NewPin->PinFriendlyName = FText::Format(LOCTEXT("UnsupportedPinFriendlyName", "(Unsupported) {0}"), FText::FromName(NewPin->GetFName()));
 		}
+
+		// Continue iterating
+		return true;
 	});
 	
 	HideEventPins(nullptr);
@@ -181,6 +187,9 @@ void UK2Node_EnhancedInputAction::HideEventPins(UEdGraphPin* RetainPin)
 			
 			Pin->bAdvancedView = !(GetDefault<UEnhancedInputEditorSettings>()->VisibleEventPinsByDefault & static_cast<uint8>(Event)) || !bIsSupported;
 		}
+
+		// Continue iterating
+		return true;
 	});
 }
 
@@ -350,20 +359,19 @@ void UK2Node_EnhancedInputAction::ExpandNode(FKismetCompilerContext& CompilerCon
 	};
 
 	const ETriggerEventsSupported SupportedTriggerEvents = GetDefault<UBlueprintEditorSettings>()->bEnableInputTriggerSupportWarnings ? InputAction->GetSupportedTriggerEvents() : ETriggerEventsSupported::All;
-	
+		
 	TArray<ActivePinData> ActivePins;
-	ForEachEventPinName([this, &ActivePins, &SupportedTriggerEvents, &CompilerContext](ETriggerEvent Event, FName PinName) 
+	ForEachActiveEventPin([this, &ActivePins, &SupportedTriggerEvents, &CompilerContext](ETriggerEvent Event, UEdGraphPin& InputActionPin)
 	{
-		UEdGraphPin* InputActionPin = FindPin(PinName, EEdGraphPinDirection::EGPD_Output);
-		if (InputActionPin && InputActionPin->LinkedTo.Num() > 0)
+		ActivePins.Add(ActivePinData(&InputActionPin, Event));
+		// Check if this exec pin is supported!
+		if (UE::Input::bShouldWarnOnUnsupportedInputPin && !UInputTrigger::IsSupportedTriggerEvent(SupportedTriggerEvents, Event))
 		{
-			ActivePins.Add(ActivePinData(InputActionPin, Event));
-			// Check if this exec pin is supported!
-			if (UE::Input::bShouldWarnOnUnsupportedInputPin && !UInputTrigger::IsSupportedTriggerEvent(SupportedTriggerEvents, Event))
-			{
-				CompilerContext.MessageLog.Warning(*FText::Format(LOCTEXT("UnsuportedEventTypeOnAction", "'{0}'on @@ may not be executed because it is not a supported trigger on this action!"), InputActionPin->GetDisplayName()).ToString(), this);
-			}
+			CompilerContext.MessageLog.Warning(*FText::Format(LOCTEXT("UnsuportedEventTypeOnAction", "'{0}'on @@ may not be executed because it is not a supported trigger on this action!"), InputActionPin.GetDisplayName()).ToString(), this);
 		}
+
+		// Continue iterating
+		return true;
 	});
 
 	if (ActivePins.Num() == 0)
@@ -549,5 +557,31 @@ TSharedPtr<FEdGraphSchemaAction> UK2Node_EnhancedInputAction::GetEventNodeAction
 	return EventNodeAction;
 }
 
+bool UK2Node_EnhancedInputAction::HasAnyConnectedEventPins() const
+{
+	bool bHasAnyConnectedEventPins = false;
+	ForEachActiveEventPin([&bHasAnyConnectedEventPins](ETriggerEvent, UEdGraphPin& Pin)
+	{
+		// Stop iterating on the first active pin
+		bHasAnyConnectedEventPins = true;
+		return false;
+	});
+
+	return bHasAnyConnectedEventPins;
+}
+
+void UK2Node_EnhancedInputAction::ForEachActiveEventPin(TFunctionRef<bool(ETriggerEvent, UEdGraphPin&)> Predicate) const
+{
+	ForEachEventPinName([this, Predicate](ETriggerEvent Event, FName PinName) 
+	{
+		UEdGraphPin* InputActionPin = FindPin(PinName, EEdGraphPinDirection::EGPD_Output);
+		if (InputActionPin && InputActionPin->LinkedTo.Num() > 0)
+		{
+			return Predicate(Event, *InputActionPin);
+		}
+		return true;
+	});
+}
+
 #undef LOCTEXT_NAMESPACE
 
diff --git a/Engine/Plugins/EnhancedInput/Source/InputBlueprintNodes/Public/K2Node_EnhancedInputAction.h b/Engine/Plugins/EnhancedInput/Source/InputBlueprintNodes/Public/K2Node_EnhancedInputAction.h
--- a/Engine/Plugins/EnhancedInput/Source/InputBlueprintNodes/Public/K2Node_EnhancedInputAction.h
+++ b/Engine/Plugins/EnhancedInput/Source/InputBlueprintNodes/Public/K2Node_EnhancedInputAction.h
@@ -75,10 +75,15 @@ class INPUTBLUEPRINTNODES_API UK2Node_EnhancedInputAction : public UK2Node, publ
 	virtual TSharedPtr<FEdGraphSchemaAction> GetEventNodeAction(const FText& ActionCategory) override;
 	//~ End IK2Node_EventNodeInterface Interface.
 
+	bool HasAnyConnectedEventPins() const;
+
 private:
 	FName GetActionName() const;
 	void HideEventPins(UEdGraphPin* RetainPin);
 
+	// Iterates each connected trigger event pin until Predicate returns false or we've iterated all active pins 
+	void ForEachActiveEventPin(TFunctionRef<bool(ETriggerEvent, UEdGraphPin&)> Predicate) const;
+
 	/** Gets the ETriggerEvent from an exec pin based on the Pins name. */
 	const ETriggerEvent GetTriggerTypeFromExecPin(const UEdGraphPin* ExecPin) const;
 	
