UIse 1.0 as the effective time dialation if there is no valid UEnhancedPlayerInput object given, which can happen in automation tests.

#rnx

diff --git a/Engine/Plugins/EnhancedInput/Source/EnhancedInput/Private/InputTriggers.cpp b/Engine/Plugins/EnhancedInput/Source/EnhancedInput/Private/InputTriggers.cpp
--- a/Engine/Plugins/EnhancedInput/Source/EnhancedInput/Private/InputTriggers.cpp
+++ b/Engine/Plugins/EnhancedInput/Source/EnhancedInput/Private/InputTriggers.cpp
@@ -123,15 +123,12 @@ ETriggerState UInputTriggerTimedBase::UpdateState_Implementation(const UEnhanced
 
 float UInputTriggerTimedBase::CalculateHeldDuration(const UEnhancedPlayerInput* const PlayerInput, const float DeltaTime) const
 {
-	if (ensureMsgf(PlayerInput, TEXT("No Player Input was given to Calculate with! Returning 1.0")))
-	{
-		const float TimeDilation = PlayerInput->GetEffectiveTimeDilation();
+	// We may not have a PlayerInput object during automation tests, so default to 1.0f if we don't have one.
+	// This will mean that TimeDilation has no effect.
+	const float TimeDilation = PlayerInput ? PlayerInput->GetEffectiveTimeDilation() : 1.0f;
 	
-		// Calculates the new held duration, applying time dilation if desired
-		return HeldDuration + (!bAffectedByTimeDilation ? DeltaTime : DeltaTime * TimeDilation);
-	}
-
-	return 1.0f;
+	// Calculates the new held duration, applying time dilation if desired
+	return HeldDuration + (!bAffectedByTimeDilation ? DeltaTime : DeltaTime * TimeDilation);
 }
 
 
