Fix a bug where a 1D axis did not respect the sign of the input value and would always be positive or 0

Added some unit tests to trest this behavior.

#rnx

diff --git a/Engine/Plugins/EnhancedInput/Source/EnhancedInput/Private/InputModifiers.cpp b/Engine/Plugins/EnhancedInput/Source/EnhancedInput/Private/InputModifiers.cpp
--- a/Engine/Plugins/EnhancedInput/Source/EnhancedInput/Private/InputModifiers.cpp
+++ b/Engine/Plugins/EnhancedInput/Source/EnhancedInput/Private/InputModifiers.cpp
@@ -174,10 +174,11 @@ FInputActionValue UInputModifierDeadZone::ModifyRaw_Implementation(const UEnhanc
 	{
 		// If the value is less then our lower threshold, return zero
 		// otherwise, clamp the value to the upper threshold.
+		const float AbsValue = FMath::Abs(AxisVal);
 		return
-			AxisVal < LowerThreshold ?
+			AbsValue < LowerThreshold ?
 				0.0f :
-				FMath::Min(AxisVal, UpperThreshold);
+				FMath::Min(AbsValue, UpperThreshold) * FMath::Sign(AxisVal);
 	};
 	
 	FVector NewValue = CurrentValue.Get<FVector>();
diff --git a/Engine/Plugins/EnhancedInput/Source/InputEditor/Private/Tests/InputModifierTest.cpp b/Engine/Plugins/EnhancedInput/Source/InputEditor/Private/Tests/InputModifierTest.cpp
--- a/Engine/Plugins/EnhancedInput/Source/InputEditor/Private/Tests/InputModifierTest.cpp
+++ b/Engine/Plugins/EnhancedInput/Source/InputEditor/Private/Tests/InputModifierTest.cpp
@@ -204,6 +204,48 @@ bool FInputModifierDeadzoneTest::RunTest(const FString& Parameters)
 	return true;
 }
 
+IMPLEMENT_SIMPLE_AUTOMATION_TEST(FInputModifierUnscaledRadialDeadzoneTest, "Input.Modifiers.UnscaledRadialDeadZone", BasicModifierTestFlags)
+bool FInputModifierUnscaledRadialDeadzoneTest::RunTest(const FString& Parameters)
+{
+	// A 1D test axis with a positive value
+	{
+		UControllablePlayer& Data =
+			GIVEN(ABasicModifierTest(this, EInputActionValueType::Axis1D));
+		AND(AnActionIsMappedToAKey(Data, TestContext, TestAction, TestAxis2));
+
+		UInputModifierDeadZone* DeadZone =
+			AND(Cast<UInputModifierDeadZone>(AModifierIsAppliedToAnAction(Data, NewObject<UInputModifierDeadZone>(), TestAction)));
+		DeadZone->LowerThreshold = 0.1f;
+		DeadZone->Type = EDeadZoneType::UnscaledRadial;
+
+		// Provide a negative value which should still be within the deadzone
+		WHEN(AKeyIsActuated(Data, TestAxis2, 0.5f));
+		AND(InputIsTicked(Data));
+		THEN(PressingKeyTriggersAction(Data, TestAction));
+		AND(TestEqual(TEXT("Positive Value is correct"), GetActionValue(Data, TestAction), 0.5f));
+	}	
+
+	// A 1D test axis with a negative value
+	{
+		UControllablePlayer& Data =
+			GIVEN(ABasicModifierTest(this, EInputActionValueType::Axis1D));
+		AND(AnActionIsMappedToAKey(Data, TestContext, TestAction, TestAxis2));
+
+		UInputModifierDeadZone* DeadZone =
+			AND(Cast<UInputModifierDeadZone>(AModifierIsAppliedToAnAction(Data, NewObject<UInputModifierDeadZone>(), TestAction)));
+		DeadZone->LowerThreshold = 0.1f;
+		DeadZone->Type = EDeadZoneType::UnscaledRadial;
+
+		// Provide a negative value which should still be within the deadzone
+		WHEN(AKeyIsActuated(Data, TestAxis2, -0.5f));
+		AND(InputIsTicked(Data));
+		THEN(PressingKeyTriggersAction(Data, TestAction));
+		AND(TestEqual(TEXT("Value is correct"), GetActionValue(Data, TestAction), -0.5f));
+	}
+	
+	return true;
+}
+
 /**
  * A simple "stress test" for Enhanced Input modifiers that we can use to measure the performance
  * of applying multiple UInputModifiers to a single key mapping.
