[Enhanced Input] Remove the native C++ setting of the default value for input modes in favor of setting them via config. This avoids a call to the gameplay tag manager during editor boot and can save ~1 second.

#rnx

diff --git a/Engine/Config/BaseInput.ini b/Engine/Config/BaseInput.ini
--- a/Engine/Config/BaseInput.ini
+++ b/Engine/Config/BaseInput.ini
@@ -253,3 +253,8 @@ HistoryColor=(R=180,G=180,B=180,A=255)
 AutoCompleteCommandColor=(B=185,G=109,R=144,A=255)
 AutoCompleteCVarColor=(B=86,G=158,R=86,A=255)
 AutoCompleteFadedColor=(R=100,G=100,B=100,A=255)
+
+
+[/Script/EnhancedInput.EnhancedInputDeveloperSettings]
+DefaultMappingContextInputModeQuery=(TokenStreamVersion=0,TagDictionary=((TagName="EnhancedInput.Modes.Default")),QueryTokenStream=(0,1,2,1,0),UserDescription="",AutoDescription="")
+DefaultInputMode=(GameplayTags=((TagName="EnhancedInput.Modes.Default")))
\ No newline at end of file
diff --git a/Engine/Plugins/EnhancedInput/Source/EnhancedInput/Private/EnhancedInputDeveloperSettings.cpp b/Engine/Plugins/EnhancedInput/Source/EnhancedInput/Private/EnhancedInputDeveloperSettings.cpp
--- a/Engine/Plugins/EnhancedInput/Source/EnhancedInput/Private/EnhancedInputDeveloperSettings.cpp
+++ b/Engine/Plugins/EnhancedInput/Source/EnhancedInput/Private/EnhancedInputDeveloperSettings.cpp
@@ -21,6 +21,4 @@ UEnhancedInputDeveloperSettings::UEnhancedInputDeveloperSettings(const FObjectIn
 	, bEnableWorldSubsystem(false)
 {
 	PlatformSettings.Initialize(UEnhancedInputPlatformSettings::StaticClass());
-	DefaultMappingContextInputModeQuery = FGameplayTagQuery::BuildQuery(FGameplayTagQueryExpression().AllTagsMatch().AddTag(UE::EnhancedInput::InputMode_Default));
-	DefaultInputMode = FGameplayTagContainer(UE::EnhancedInput::InputMode_Default);
 }
\ No newline at end of file
