[Enhanced Input] Remove override UEnhancedplayerInput::ProcessInputStack which just called the Super function anyway.

#rnx

diff --git a/Engine/Plugins/EnhancedInput/Source/EnhancedInput/Private/EnhancedPlayerInput.cpp b/Engine/Plugins/EnhancedInput/Source/EnhancedInput/Private/EnhancedPlayerInput.cpp
--- a/Engine/Plugins/EnhancedInput/Source/EnhancedInput/Private/EnhancedPlayerInput.cpp
+++ b/Engine/Plugins/EnhancedInput/Source/EnhancedInput/Private/EnhancedPlayerInput.cpp
@@ -792,11 +792,6 @@ void UEnhancedPlayerInput::EvaluateInputDelegates(const TArray<UInputComponent*>
 	Super::EvaluateInputDelegates(InputComponentStack, DeltaTime, bGamePaused, KeysWithEvents);
 }
 
-void UEnhancedPlayerInput::ProcessInputStack(const TArray<UInputComponent*>& InputComponentStack, const float DeltaTime, const bool bGamePaused)
-{
-	Super::ProcessInputStack(InputComponentStack, DeltaTime, bGamePaused);
-}
-
 void UEnhancedPlayerInput::FlushPressedKeys()
 {
 	Super::FlushPressedKeys();
diff --git a/Engine/Plugins/EnhancedInput/Source/EnhancedInput/Public/EnhancedPlayerInput.h b/Engine/Plugins/EnhancedInput/Source/EnhancedInput/Public/EnhancedPlayerInput.h
--- a/Engine/Plugins/EnhancedInput/Source/EnhancedInput/Public/EnhancedPlayerInput.h
+++ b/Engine/Plugins/EnhancedInput/Source/EnhancedInput/Public/EnhancedPlayerInput.h
@@ -119,9 +119,6 @@ public:
 	void InjectInputForAction(TObjectPtr<const UInputAction> Action, FInputActionValue RawValue, const TArray<UInputModifier*>& Modifiers = {}, const TArray<UInputTrigger*>& Triggers = {});
 
 	virtual bool InputKey(const FInputKeyParams& Params) override;
-	
-	// Applies modifiers and triggers without affecting keys read by the base input system
-	virtual void ProcessInputStack(const TArray<UInputComponent*>& InputComponentStack, const float DeltaTime, const bool bGamePaused) override;
 
 	/** Returns the Time Dilation value that is currently effecting this input. */
 	float GetEffectiveTimeDilation() const;
