[Enhanced Input] Make GetModifiers/GetTrigger functions return the actual TObjectPtr array type instead of raw pointers.

#rnx

diff --git a/Engine/Plugins/EnhancedInput/Source/EnhancedInput/Public/InputAction.h b/Engine/Plugins/EnhancedInput/Source/EnhancedInput/Public/InputAction.h
--- a/Engine/Plugins/EnhancedInput/Source/EnhancedInput/Public/InputAction.h
+++ b/Engine/Plugins/EnhancedInput/Source/EnhancedInput/Public/InputAction.h
@@ -256,11 +256,11 @@ public:
 	// Time that this action was last actively triggered
 	float GetLastTriggeredWorldTime() const { return LastTriggeredWorldTime; }
 
-	const TArray<UInputTrigger*>& GetTriggers() const { return Triggers; }
-	const TArray<UInputModifier*>& GetModifiers() const { return Modifiers; }
+	const TArray<TObjectPtr<UInputTrigger>>& GetTriggers() const { return Triggers; }
+	const TArray<TObjectPtr<UInputModifier>>& GetModifiers() const { return Modifiers; }
 
 	// The source action that this instance is created from
-	const UInputAction* GetSourceAction() const { return SourceAction; }
+	TObjectPtr<const UInputAction> GetSourceAction() const { return SourceAction; }
 
 	void AddInputTrigger(UInputTrigger* InputTrigger) { Triggers.Add(InputTrigger); }
 	void RemoveInputTrigger(UInputTrigger* InputTrigger) { Triggers.Remove(InputTrigger); }
