[Enhanced Input]  Fix the "HasBindings" function hiding its parent class impl. This function should have been virtual.

#rnx

diff --git a/Engine/Plugins/EnhancedInput/Source/EnhancedInput/Public/EnhancedInputComponent.h b/Engine/Plugins/EnhancedInput/Source/EnhancedInput/Public/EnhancedInputComponent.h
--- a/Engine/Plugins/EnhancedInput/Source/EnhancedInput/Public/EnhancedInputComponent.h
+++ b/Engine/Plugins/EnhancedInput/Source/EnhancedInput/Public/EnhancedInputComponent.h
@@ -272,7 +272,7 @@ public:
 		return TUniquePtr<FEnhancedInputActionEventBinding>(new FEnhancedInputActionEventDelegateBinding<TSignature>(*this, EInputBindingClone::ForceClone));
 	}
 
-	virtual void SetShouldFireWithEditorScriptGuard(const bool bNewValue)
+	virtual void SetShouldFireWithEditorScriptGuard(const bool bNewValue) override
 	{
 		Delegate.SetShouldFireWithEditorScriptGuard(bNewValue);
 	}
@@ -384,7 +384,7 @@ public:
 	 *
 	 * @return true if any bindings are set, false otherwise.
 	 */
-	bool HasBindings() const;
+	virtual bool HasBindings() const override;
 
 	/**
 	 * Removes all action bindings.
diff --git a/Engine/Source/Runtime/Engine/Classes/Components/InputComponent.h b/Engine/Source/Runtime/Engine/Classes/Components/InputComponent.h
--- a/Engine/Source/Runtime/Engine/Classes/Components/InputComponent.h
+++ b/Engine/Source/Runtime/Engine/Classes/Components/InputComponent.h
@@ -801,7 +801,7 @@ public:
 	 *
 	 * @return true if any bindings are set, false otherwise.
 	 */
-	ENGINE_API bool HasBindings() const;
+	ENGINE_API virtual bool HasBindings() const;
 
 
 	/**
