[Enhanced Input] Return the "real" value of an input action instance even when it is not in the explict state of "triggered". This makes for a more intuitive use of bool (digital) actions. Before, you could listen for the "started' event of a bool action, and the value of that input action would be "false" because it is not in the "triggered' phase.

This can be disabled and reverted back to the previous behavior by setting the CVar "EnhancedInput.bAlwaysGetRealValueFromActionInstanceData" to false.

diff --git a/Engine/Plugins/EnhancedInput/Source/EnhancedInput/Private/InputAction.cpp b/Engine/Plugins/EnhancedInput/Source/EnhancedInput/Private/InputAction.cpp
--- a/Engine/Plugins/EnhancedInput/Source/EnhancedInput/Private/InputAction.cpp
+++ b/Engine/Plugins/EnhancedInput/Source/EnhancedInput/Private/InputAction.cpp
@@ -2,6 +2,7 @@
 
 #include "InputAction.h"
 
+#include "HAL/IConsoleManager.h"
 #include "PlayerMappableKeySettings.h"
 
 #if WITH_EDITOR
@@ -12,6 +13,14 @@
 
 #define LOCTEXT_NAMESPACE "EnhancedInputAction"
 
+namespace UE::EnhancedInput
+{
+	static bool bAlwaysGetRealValueFromActionInstanceData = 1;
+	static FAutoConsoleVariableRef CVarAlwaysGetRealValueFromActionInstanceData(TEXT("EnhancedInput.bAlwaysGetRealValueFromActionInstanceData"),
+		bAlwaysGetRealValueFromActionInstanceData,
+		TEXT("If true, then input action instances will always return their true value instead of being conditoinal on the trigger state. Note: This CVar will eventually be removed"));
+}
+
 FInputActionInstance::FInputActionInstance(const UInputAction* InSourceAction)
 	: SourceAction(InSourceAction)
 {
@@ -39,6 +48,22 @@ FInputActionInstance::FInputActionInstance(const UInputAction* InSourceAction)
 	}
 }
 
+FInputActionValue FInputActionInstance::GetValue() const
+{
+	if (UE::EnhancedInput::bAlwaysGetRealValueFromActionInstanceData)
+	{
+		return Value;
+	}
+	
+	// The legacy behavior only returned the true value when in the "Triggered" state
+	// which would result in bool (digital) input actions being false on the "started" phase,
+	// which is confusing. 
+	return
+		TriggerEvent == ETriggerEvent::Triggered ?
+		Value :
+		FInputActionValue(Value.GetValueType(), FInputActionValue::Axis3D::ZeroVector);
+ } 
+
 // Calculate a collective representation of trigger state from evaluations of all triggers in one or more trigger groups.
 ETriggerState FTriggerStateTracker::EvaluateTriggers(const UEnhancedPlayerInput* PlayerInput, const TArray<UInputTrigger*>& Triggers, FInputActionValue ModifiedValue, float DeltaTime)
 {
diff --git a/Engine/Plugins/EnhancedInput/Source/EnhancedInput/Public/InputAction.h b/Engine/Plugins/EnhancedInput/Source/EnhancedInput/Public/InputAction.h
--- a/Engine/Plugins/EnhancedInput/Source/EnhancedInput/Public/InputAction.h
+++ b/Engine/Plugins/EnhancedInput/Source/EnhancedInput/Public/InputAction.h
@@ -245,7 +245,7 @@ public:
 	ETriggerEvent GetTriggerEvent() const { return TriggerEvent; }
 
 	// Current action value - Will be zero if the current trigger event is not ETriggerEvent::Triggered!
-	FInputActionValue GetValue() const { return TriggerEvent == ETriggerEvent::Triggered ? Value : FInputActionValue(Value.GetValueType(), FInputActionValue::Axis3D::ZeroVector); }
+	FInputActionValue GetValue() const;
 
 	// Total time the action has been evaluating triggering (Ongoing & Triggered)
 	float GetElapsedTime() const { return ElapsedProcessedTime; }
