[Enhanced Input] When clearing any forced FKey values use the proper way of clearing paired axis keys instead of directly calling InputKey. This is to keep the behavior in line with what 40388643 does.

#rnx
#okforigthub public

diff --git a/Engine/Plugins/EnhancedInput/Source/EnhancedInput/Private/EnhancedInputSubsystemInterface.cpp b/Engine/Plugins/EnhancedInput/Source/EnhancedInput/Private/EnhancedInputSubsystemInterface.cpp
--- a/Engine/Plugins/EnhancedInput/Source/EnhancedInput/Private/EnhancedInputSubsystemInterface.cpp
+++ b/Engine/Plugins/EnhancedInput/Source/EnhancedInput/Private/EnhancedInputSubsystemInterface.cpp
@@ -13,6 +13,7 @@
 #include "PlayerMappableInputConfig.h"
 #include "PlayerMappableKeySettings.h"
 #include "Engine/LocalPlayer.h"
+#include "Engine/World.h"
 #include "GenericPlatform/GenericPlatformInputDeviceMapper.h"
 
 #include UE_INLINE_GENERATED_CPP_BY_NAME(EnhancedInputSubsystemInterface)
@@ -1124,13 +1125,14 @@ void IEnhancedInputSubsystemInterface::RebuildControlMappings()
 }
 
 template<typename T>
-void InjectKey(T* InjectVia, FKey Key, const FInputActionValue& Value, float DeltaTime, const FPlatformUserId OwningUser)
+void InjectKey(T* InjectVia, FKey Key, const FInputActionValue& Value, float DeltaTime, const FPlatformUserId OwningUser, const EInputEvent Event = IE_Pressed)
 {
 	const FInputDeviceId DeviceToSimulate = IPlatformInputDeviceMapper::Get().GetPrimaryInputDeviceForUser(OwningUser);
 	
-	auto SimulateKeyPress = [InjectVia, DeltaTime, DeviceToSimulate](const FKey& KeyToSim, const float Value)
+	auto SimulateKeyPress = [InjectVia, DeltaTime, DeviceToSimulate, Event](const FKey& KeyToSim, const float Value)
 	{
 		FInputKeyParams Params;
+		Params.Event = Event;
 		Params.Delta.X = Value;
 		Params.Key = KeyToSim;
 		Params.NumSamples = KeyToSim.IsAnalog() ? 1 : 0;
@@ -1245,8 +1247,9 @@ void IEnhancedInputSubsystemInterface::RemoveForcedInput(FKey Key)
 	// Otherwise, we need to inject a release event tos player input
 	if (UEnhancedPlayerInput* PlayerInput = GetPlayerInput())
 	{
-		FInputKeyParams Params;
-		Params.Key = Key;
+		// Set the input device id to the platform user's default input device
+		const FPlatformUserId UserId = PlayerInput->GetOwningLocalPlayer()->GetPlatformUserId();
+		const float DeltaTime = PlayerInput->GetWorld()->GetDeltaSeconds();
 
 		// We want to inject the opposite of whatever we were previously injecting for this key
 		// in order to get it back to providing a fake value of zero. For example, if we were injecting (.5,.5)
@@ -1254,22 +1257,16 @@ void IEnhancedInputSubsystemInterface::RemoveForcedInput(FKey Key)
 		// for analog keys.
 		//
 		// Any digital key we always want a value of zero to ensure it is treated as a release event.
-		Params.Delta = Key.IsAnalog() ? -InjectedKeyData->LastInjectedValue.Get<FVector>() : FVector::ZeroVector;
-		Params.Event = EInputEvent::IE_Released;
-		Params.NumSamples = Key.IsAnalog() ? 1 : 0;
-
-		// Set the input device id to the platform user's default input device
-		const FPlatformUserId UserId = PlayerInput->GetOwningLocalPlayer()->GetPlatformUserId();
-		Params.InputDevice = IPlatformInputDeviceMapper::Get().GetPrimaryInputDeviceForUser(UserId);
+		const FInputActionValue ValueToInject = Key.IsAnalog() ? -InjectedKeyData->LastInjectedValue.Get<FVector>() : FVector::ZeroVector;
 	
 		// Prefer sending the key released event via a player controller if one is available.
 		if (APlayerController* Controller = Cast<APlayerController>(PlayerInput->GetOuter()))
 		{
-			Controller->InputKey(Params);
+			InjectKey(Controller, Key, ValueToInject, DeltaTime, UserId, IE_Released);
 		}
 		else
 		{
-			PlayerInput->InputKey(Params);
+			InjectKey(PlayerInput, Key, ValueToInject, DeltaTime, UserId, IE_Released);
 		}
 	
 		// Flush the player's pressed keys to ensure that the removed event is read
