[Enhanced Input] Fix a bug where input modifiers/triggers would incorrectly evaluate themselves as being the same just because they were of the same type, not checking if the default values were actually the same or not.

You could repro this issue by making two IMC assets with the same key mappings and a scalar modifier. One IMC_A, set the scalar value ot 1. On IMC_B, set it to 100. Add IMC_A with a priority of 0, and then later add IMC_B with a higher priority. The buggy result is that the scalar value of 1 would still be used, even though IMC_B has been pushed and has scalara values of 1. This was happening because when the control mappings were rebuilt, the deep copy of the original IMC_A's modifier would retain their state because the class was the same. The expected result that we want to reset the state of that modifier to use the newer IMC_B's values on the scalar modifier instead.

diff --git a/Engine/Plugins/EnhancedInput/Source/EnhancedInput/Public/EnhancedActionKeyMapping.h b/Engine/Plugins/EnhancedInput/Source/EnhancedInput/Public/EnhancedActionKeyMapping.h
--- a/Engine/Plugins/EnhancedInput/Source/EnhancedInput/Public/EnhancedActionKeyMapping.h
+++ b/Engine/Plugins/EnhancedInput/Source/EnhancedInput/Public/EnhancedActionKeyMapping.h
@@ -189,8 +189,19 @@ public:
 				// Both are nullptr. Consider that as same type.
 				continue;
 			}
-			if (ObjectA->GetClass() != ObjectB->GetClass())
+			
+			const UClass* ClassA = ObjectA->GetClass();
+			const UClass* ClassB = ObjectB->GetClass();
+			
+			if (ClassA != ClassB)
 			{
+				// If the classes differ then they are not the same
+				return false;
+			}
+			
+			if (ClassA->GetDefaultObject() != ClassB->GetDefaultObject())
+			{
+				// If the default objects of these two objects differ, then they should be treated differently
 				return false;
 			}
 		}
