balanced macro fix

#rnx

diff --git a/Engine/Plugins/EnhancedInput/Source/EnhancedInput/Private/UserSettings/EnhancedInputUserSettings.cpp b/Engine/Plugins/EnhancedInput/Source/EnhancedInput/Private/UserSettings/EnhancedInputUserSettings.cpp
--- a/Engine/Plugins/EnhancedInput/Source/EnhancedInput/Private/UserSettings/EnhancedInputUserSettings.cpp
+++ b/Engine/Plugins/EnhancedInput/Source/EnhancedInput/Private/UserSettings/EnhancedInputUserSettings.cpp
@@ -1520,7 +1520,7 @@ UEnhancedPlayerMappableKeyProfile* UEnhancedInputUserSettings::GetKeyProfileWith
 {
 	return GetKeyProfileWithId(ProfileId.ToString());
 }
-PRAGMA_DISABLE_DEPRECATION_WARNINGS
+PRAGMA_ENABLE_DEPRECATION_WARNINGS
 
 UEnhancedPlayerMappableKeyProfile* UEnhancedInputUserSettings::GetKeyProfileWithId(const FString& ProfileId) const
 {
