[Enhanced Input] Add a unit test for creating multiple key profiles and swapping between them.

Remove an unnecessary check for a valid ULocalPlayer outer for the key profiles.

diff --git a/Engine/Plugins/EnhancedInput/Source/EnhancedInput/Private/UserSettings/EnhancedInputUserSettings.cpp b/Engine/Plugins/EnhancedInput/Source/EnhancedInput/Private/UserSettings/EnhancedInputUserSettings.cpp
--- a/Engine/Plugins/EnhancedInput/Source/EnhancedInput/Private/UserSettings/EnhancedInputUserSettings.cpp
+++ b/Engine/Plugins/EnhancedInput/Source/EnhancedInput/Private/UserSettings/EnhancedInputUserSettings.cpp
@@ -1378,12 +1378,6 @@ PRAGMA_ENABLE_DEPRECATION_WARNINGS
 
 bool UEnhancedInputUserSettings::SetActiveKeyProfile(const FString& InProfileId)
 {
-	if (!GetLocalPlayer())
-	{
-		UE_LOG(LogEnhancedInput, Log, TEXT("Failed to find the Local Player associated with the Enhanced Input user settings!"));
-		return false;
-	}
-
 	const FString OriginalProfileId = GetActiveKeyProfileId();
 	UEnhancedPlayerMappableKeyProfile* OriginalProfile = GetActiveKeyProfile();
 
diff --git a/Engine/Plugins/EnhancedInput/Source/InputEditor/Private/Tests/InputPlayerMappableKeysTests.cpp b/Engine/Plugins/EnhancedInput/Source/InputEditor/Private/Tests/InputPlayerMappableKeysTests.cpp
--- a/Engine/Plugins/EnhancedInput/Source/InputEditor/Private/Tests/InputPlayerMappableKeysTests.cpp
+++ b/Engine/Plugins/EnhancedInput/Source/InputEditor/Private/Tests/InputPlayerMappableKeysTests.cpp
@@ -43,7 +43,7 @@ UControllablePlayer& ABasicPlayerMappableKeysTest(FAutomationTestBase* Test)
 }
 
 // Provides a test player with a single input action (TestAction) added to it mapped to X
-UControllablePlayer& APlayerMappableKeysTest_WithAction(FAutomationTestBase* Test)
+UControllablePlayer& APlayerMappableKeysTest_WithAction(FAutomationTestBase* Test, const FKey& DefaultKeyMapping = UE::Input::TestKeyX)
 {
 	UWorld* World =
 	GIVEN(AnEmptyWorld());
@@ -61,14 +61,14 @@ UControllablePlayer& APlayerMappableKeysTest_WithAction(FAutomationTestBase* Tes
 	TObjectPtr<UInputMappingContext>* IMC = Data.InputContext.Find(TestContext);
 	Test->TestNotNull("Mock Mapping Context", IMC);
 
-	// Add a simple mapping from test action to X
+	// Add a simple mapping from test action to the given default key
 	UInputAction* Action1 =
 	AND(AnInputAction(Data, TestAction, EInputActionValueType::Axis2D));
 
 	Test->TestNotNull("Mock Input Action", Action1);
 
 	// Test mapping to X
-	AND(AnActionIsMappedToAKey(Data, TestContext, TestAction, UE::Input::TestKeyX));
+	AND(AnActionIsMappedToAKey(Data, TestContext, TestAction, DefaultKeyMapping));
 	
 	const bool bRes = Settings->RegisterInputMappingContext(*IMC);
 	Test->TestTrue("Registered Mock IMC", bRes);
@@ -402,5 +402,131 @@ bool FResetToDefaultTest::RunTest(const FString& Parameters)
 		TestEqual("reset to default correctly!", MappedKeys[0], UE::Input::TestKeyX);
 	}
 	
+	return true;
+}
+
+// A test to create multiple key mapping profiles
+IMPLEMENT_SIMPLE_AUTOMATION_TEST(FMultipleKeyProfilesTest, "Input.PlayerMappableKeys.MultipleProfiles", BasicPlayerMappableKeysTestFlags)
+bool FMultipleKeyProfilesTest::RunTest(const FString& Parameters)
+{
+	// Set up the test
+	UControllablePlayer& Data =
+		GIVEN(APlayerMappableKeysTest_WithAction(this));
+	
+	UEnhancedInputUserSettings* Settings = Data.Subsystem->GetUserSettings();
+	TestNotNull("Mock Settings", Settings);
+	
+	UEnhancedPlayerMappableKeyProfile* Profile_A = Settings->GetActiveKeyProfile();
+	TestNotNull("Current Key Profile", Profile_A);	
+	
+	const FString Profile_A_ID = Profile_A->GetProfileIdString();
+	TestFalse("Current Key Profile has a valid string ID", Profile_A_ID.IsEmpty());
+
+	// Map the "Z" key to the test action on the default profile
+	{
+		FGameplayTagContainer ErrorReason;
+
+		// Remap the action in slot 1 from x to z
+		FMapPlayerKeyArgs Args = {};
+		Args.MappingName = TestAction;
+		Args.NewKey = UE::Input::TestKeyZ;
+		Args.Slot = EPlayerMappableKeySlot::First;
+
+		Settings->MapPlayerKey(Args, OUT ErrorReason);
+		TestTrue("Successful Remap", ErrorReason.IsEmpty());
+
+		TArray<FKey> MappedKeys;
+
+		FPlayerMappableKeyQueryOptions Opts = {};
+		Opts.MappingName = TestAction;
+
+		Profile_A->QueryPlayerMappedKeys(Opts, OUT MappedKeys);
+		TestTrue("Query Mapped Keys populated the array of keys", !MappedKeys.IsEmpty());
+
+		TestTrue("Default profile correctly mapped key to TestKeyZ", MappedKeys[0] == UE::Input::TestKeyZ);
+
+		// Ensure that the mapping is now mapped to the Z key
+		const FKeyMappingRow* Row = Profile_A->FindKeyMappingRow(TestAction);
+		TestTrue("Row has mappings", Row->Mappings.Num() == 1);
+
+		const FPlayerKeyMapping& SingleMapping = *Row->Mappings.CreateConstIterator();
+		
+		TestTrue("Mapping is correctly marked as dirty", SingleMapping.IsDirty());
+		TestTrue("Mapping is in the correct slot", SingleMapping.GetSlot() == EPlayerMappableKeySlot::First);
+		TestTrue("Mapping has the right custom FKey", SingleMapping.GetCurrentKey() == UE::Input::TestKeyZ);
+		TestTrue("Mapping has the right default FKey", SingleMapping.GetDefaultKey() == UE::Input::TestKeyX);
+	}
+
+	// Create a new key profile and set it as the currently active one.
+	UEnhancedPlayerMappableKeyProfile* Profile_B = nullptr;
+	const FString Profile_B_Requested_ID = TEXT("input.test.profile.b");
+	{
+		FPlayerMappableKeyProfileCreationArgs Args = {};
+		Args.ProfileStringIdentifier = Profile_B_Requested_ID;
+		Args.bSetAsCurrentProfile = true;
+		// use the default profile type in case there is a custom one set for the project these tests are run in.
+		Args.ProfileType = UEnhancedPlayerMappableKeyProfile::StaticClass();
+		Profile_B = Settings->CreateNewKeyProfile(Args);
+
+		TestNotNull("Profile B is valid", Profile_B);
+		TestTrue("Profile B has the correct String Identifier", Profile_B->GetProfileIdString().Equals(Profile_B_Requested_ID));
+		
+		const UEnhancedPlayerMappableKeyProfile* ActiveProfile = Settings->GetActiveKeyProfile();
+		TestTrue("Profile B was set as the active profile", Profile_B == ActiveProfile);	
+
+		// Ensure that the state of the new profiles key mappings are correct
+		// Profile B should have the default key mappings set up because the customization was done on the default profile			
+		TArray<FKey> MappedKeys;
+
+		FPlayerMappableKeyQueryOptions Opts = {};
+		Opts.MappingName = TestAction;
+
+		Profile_B->QueryPlayerMappedKeys(Opts, OUT MappedKeys);
+		TestTrue("Profile B: Query Mapped Keys populated the array of keys", !MappedKeys.IsEmpty());
+		TestTrue("Profile B:correctly reset its mapping mapped key to TestKeyX", MappedKeys[0] == UE::Input::TestKeyX);
+
+		// Ensure that the mapping is reset back to the default X key
+		const FKeyMappingRow* Row = Profile_B->FindKeyMappingRow(TestAction);
+		TestTrue("Profile B: Row has mappings", Row->Mappings.Num() == 1);
+
+		const FPlayerKeyMapping& SingleMapping = *Row->Mappings.CreateConstIterator();
+
+		TestFalse("Profile B: Mapping is correctly not marked as dirty", SingleMapping.IsDirty());
+		TestTrue("Profile B: Mapping is in the correct slot", SingleMapping.GetSlot() == EPlayerMappableKeySlot::First);
+		TestTrue("Profile B: Mapping has the right custom FKey", SingleMapping.GetCurrentKey() == UE::Input::TestKeyX);
+		TestTrue("Profile B: Mapping has the right default FKey", SingleMapping.GetDefaultKey() == UE::Input::TestKeyX);
+	}
+
+	// Ensure that you can still query the keys from a non-active key profile
+	{		
+		TArray<FKey> MappedKeys;
+
+		FPlayerMappableKeyQueryOptions Opts = {};
+		Opts.MappingName = TestAction;
+
+		Profile_A->QueryPlayerMappedKeys(Opts, OUT MappedKeys);
+		TestTrue("Profile A: Query Mapped Keys populated the array of keys", !MappedKeys.IsEmpty());
+		TestTrue("Profile A:correctly reset its mapping mapped key to TestKeyX", MappedKeys[0] == UE::Input::TestKeyZ);
+
+		// Ensure that the mapping is reset back to the default X key
+		const FKeyMappingRow* Row = Profile_A->FindKeyMappingRow(TestAction);
+		TestTrue("Profile A: Row has mappings", Row->Mappings.Num() == 1);
+
+		const FPlayerKeyMapping& SingleMapping = *Row->Mappings.CreateConstIterator();
+
+		TestTrue("Profile A: Mapping is correctly not marked as dirty", SingleMapping.IsDirty());
+		TestTrue("Profile A: Mapping is in the correct slot", SingleMapping.GetSlot() == EPlayerMappableKeySlot::First);
+		TestTrue("Profile A: Mapping has the right custom FKey", SingleMapping.GetCurrentKey() == UE::Input::TestKeyZ);
+		TestTrue("Profile A: Mapping has the right default FKey", SingleMapping.GetDefaultKey() == UE::Input::TestKeyX);
+	}
+	
+	// Ensure that both profiles are registered
+	{
+		const TMap<FString, TObjectPtr<UEnhancedPlayerMappableKeyProfile>> ProfileMap = Settings->GetAllAvailableKeyProfiles();
+
+		TestTrue("Profile A is in the map", ProfileMap.Contains(Profile_A_ID));
+		TestTrue("Profile B is in the map", ProfileMap.Contains(Profile_B_Requested_ID));
+	}
+
 	return true;
 }
\ No newline at end of file
diff --git a/Engine/Plugins/EnhancedInput/Source/InputEditor/Private/Tests/InputTestFramework.cpp b/Engine/Plugins/EnhancedInput/Source/InputEditor/Private/Tests/InputTestFramework.cpp
--- a/Engine/Plugins/EnhancedInput/Source/InputEditor/Private/Tests/InputTestFramework.cpp
+++ b/Engine/Plugins/EnhancedInput/Source/InputEditor/Private/Tests/InputTestFramework.cpp
@@ -38,8 +38,11 @@ void UMockInputUserSettings::Initialize(ULocalPlayer* LP)
 	// Create a default key mapping profile in the case where one doesn't exist
 	if (!GetActiveKeyProfile())
 	{
+		const FString Profile_Default_Requested_ID = TEXT("input.test.profile.default");
+
 		FPlayerMappableKeyProfileCreationArgs Args = {};
 		Args.bSetAsCurrentProfile = true;
+		Args.ProfileStringIdentifier = Profile_Default_Requested_ID;
 	
 		CreateNewKeyProfile(Args);
 	}
