[Enhanced Input] Add a function redirect from GetCurrentKeyProfile to GetActiveKeyProfile to auto-fixup the callsites in BP.

#rnx

diff --git a/Engine/Plugins/EnhancedInput/Config/BaseEnhancedInput.ini b/Engine/Plugins/EnhancedInput/Config/BaseEnhancedInput.ini
--- a/Engine/Plugins/EnhancedInput/Config/BaseEnhancedInput.ini
+++ b/Engine/Plugins/EnhancedInput/Config/BaseEnhancedInput.ini
@@ -6,4 +6,5 @@
 +FunctionRedirects=(OldName="/Script/EnhancedInput.EnhancedInputSubsystemInterface.RemovePlayerMappedKey", NewName="/Script/EnhancedInput.EnhancedInputSubsystemInterface.K2_RemovePlayerMappedKeyInSlot")
 +FunctionRedirects=(OldName="/Script/EnhancedInput.EnhancedInputSubsystemInterface.GetPlayerMappedKey", NewName="/Script/EnhancedInput.EnhancedInputSubsystemInterface.K2_GetPlayerMappedKeyInSlot")
 ; 5.6
-+PropertyRedirects=(OldName="EnhancedInput.PlayerKeyMapping.AssociatedInputAction",NewName="EnhancedInput.PlayerKeyMapping.AssociatedInputActionSoft")
\ No newline at end of file
++PropertyRedirects=(OldName="EnhancedInput.PlayerKeyMapping.AssociatedInputAction",NewName="EnhancedInput.PlayerKeyMapping.AssociatedInputActionSoft")
++FunctionRedirects=(OldName="/Script/EnhancedInput.EnhancedInputUserSettings.GetCurrentKeyProfile", NewName="/Script/EnhancedInput.EnhancedInputUserSettings.GetActiveKeyProfile")
\ No newline at end of file
diff --git a/Engine/Plugins/EnhancedInput/Source/EnhancedInput/Public/UserSettings/EnhancedInputUserSettings.h b/Engine/Plugins/EnhancedInput/Source/EnhancedInput/Public/UserSettings/EnhancedInputUserSettings.h
--- a/Engine/Plugins/EnhancedInput/Source/EnhancedInput/Public/UserSettings/EnhancedInputUserSettings.h
+++ b/Engine/Plugins/EnhancedInput/Source/EnhancedInput/Public/UserSettings/EnhancedInputUserSettings.h
@@ -704,7 +704,6 @@ public:
 
 	/** Get the current key profile that the user has set */
 	UE_DEPRECATED(5.6, "Use GetActiveKeyProfile instead.")
-	UFUNCTION(BlueprintCallable, BlueprintPure, Category="Enhanced Input|User Settings", meta=(DeprecatedFunction, DeprecationMessage="Use GetActiveKeyProfile"))
 	UEnhancedPlayerMappableKeyProfile* GetCurrentKeyProfile() const;
 
 	/** Get the currently active key profile that the user has set */
