[Enhanced Input] More redirectors to make deprecation of gameplay tag key profiles easier

diff --git a/Engine/Plugins/EnhancedInput/Config/BaseEnhancedInput.ini b/Engine/Plugins/EnhancedInput/Config/BaseEnhancedInput.ini
--- a/Engine/Plugins/EnhancedInput/Config/BaseEnhancedInput.ini
+++ b/Engine/Plugins/EnhancedInput/Config/BaseEnhancedInput.ini
@@ -7,4 +7,10 @@
 +FunctionRedirects=(OldName="/Script/EnhancedInput.EnhancedInputSubsystemInterface.GetPlayerMappedKey", NewName="/Script/EnhancedInput.EnhancedInputSubsystemInterface.K2_GetPlayerMappedKeyInSlot")
 ; 5.6
 +PropertyRedirects=(OldName="EnhancedInput.PlayerKeyMapping.AssociatedInputAction",NewName="EnhancedInput.PlayerKeyMapping.AssociatedInputActionSoft")
-+FunctionRedirects=(OldName="/Script/EnhancedInput.EnhancedInputUserSettings.GetCurrentKeyProfile", NewName="/Script/EnhancedInput.EnhancedInputUserSettings.GetActiveKeyProfile")
\ No newline at end of file
++FunctionRedirects=(OldName="/Script/EnhancedInput.EnhancedInputUserSettings.GetCurrentKeyProfile", NewName="/Script/EnhancedInput.EnhancedInputUserSettings.GetActiveKeyProfile")
++FunctionRedirects=(OldName="/Script/EnhancedInput.EnhancedInputUserSettings.GetKeyProfileWithIdentifier", NewName="/Script/EnhancedInput.EnhancedInputUserSettings.GetKeyProfileWithId")
++FunctionRedirects=(OldName="/Script/EnhancedInput.EnhancedInputUserSettings.SetKeyProfile", NewName="/Script/EnhancedInput.EnhancedInputUserSettings.SetActiveKeyProfile")
++FunctionRedirects=(OldName="/Script/EnhancedInput.EnhancedInputUserSettings.GetCurrentKeyProfileIdentifier", NewName="/Script/EnhancedInput.EnhancedInputUserSettings.GetActiveKeyProfileId")
++FunctionRedirects=(OldName="/Script/EnhancedInput.EnhancedPlayerMappableKeyProfile.GetProfileIdentifer", NewName="/Script/EnhancedInput.EnhancedPlayerMappableKeyProfile.GetProfileIdString")
++PropertyRedirects=(OldName="EnhancedInput.MapPlayerKeyArgs.ProfileId",NewName="EnhancedInput.MapPlayerKeyArgs.ProfileIdString")
++PropertyRedirects=(OldName="EnhancedInput.PlayerMappableKeyProfileCreationArgs.ProfileIdentifier",NewName="EnhancedInput.PlayerMappableKeyProfileCreationArgs.ProfileStringIdentifier")
\ No newline at end of file
diff --git a/Engine/Plugins/EnhancedInput/Source/EnhancedInput/Public/UserSettings/EnhancedInputUserSettings.h b/Engine/Plugins/EnhancedInput/Source/EnhancedInput/Public/UserSettings/EnhancedInputUserSettings.h
--- a/Engine/Plugins/EnhancedInput/Source/EnhancedInput/Public/UserSettings/EnhancedInputUserSettings.h
+++ b/Engine/Plugins/EnhancedInput/Source/EnhancedInput/Public/UserSettings/EnhancedInputUserSettings.h
@@ -323,7 +323,6 @@ public:
 	virtual void SetDisplayName(const FText& NewDisplayName);
 
 	UE_DEPRECATED(5.6, "Use GetProfileIdString instead")
-	UFUNCTION(BlueprintCallable, BlueprintPure, Category="Enhanced Input|User Settings", meta=(DeprecatedFunction, DeprecationMessage = "Use GetProfileIdString instead"))
 	const FGameplayTag& GetProfileIdentifer() const;
 
 	UFUNCTION(BlueprintCallable, BlueprintPure, Category="Enhanced Input|User Settings")
@@ -684,7 +683,6 @@ public:
 	 * Changes the currently active key profile to the one with the given name. Returns true if the operation was successful.
 	 */
 	UE_DEPRECATED(5.6, "Use the version that takes an FString instead.")
-	UFUNCTION(BlueprintCallable, Category="Enhanced Input|User Settings", meta=(ReturnDisplayName = "Was Successful", DeprecatedFunction, DeprecationMessage="Use the version that takes an FString instead."))
 	virtual bool SetKeyProfile(const FGameplayTag& InProfileId) final;
 
 	/**
@@ -695,7 +693,6 @@ public:
 	
 	/** Gets the currently selected key profile */
 	UE_DEPRECATED(5.6, "Use GetActiveKeyProfileId instead ")
-	UFUNCTION(BlueprintCallable, BlueprintPure, Category="Enhanced Input|User Settings")
 	virtual const FGameplayTag& GetCurrentKeyProfileIdentifier() const final;
 
 	/** Gets the currently selected key profile */
@@ -741,7 +738,6 @@ public:
 
 	/** Returns the key profile with the given name if one exists. Null if one doesn't exist */
 	UE_DEPRECATED(5.6, "Use the version which takes an FString instead")
-	UFUNCTION(BlueprintCallable, BlueprintPure, Category="Enhanced Input|User Settings", meta=(DeprecatedFunction, DeprecationMessage="Use the version which takes a String instead"))
 	virtual UEnhancedPlayerMappableKeyProfile* GetKeyProfileWithIdentifier(const FGameplayTag& ProfileId) const final;
 
 	UFUNCTION(BlueprintCallable, BlueprintPure, Category="Enhanced Input|User Settings")
