Flag UPlayerMappableInputConfig_Factory as deprecated at the UObject level and filter it out from the CB add new menu since it was showing up as "Data Asset" due to a lack of an asset definition.

diff --git a/Engine/Plugins/EnhancedInput/Source/InputEditor/Private/InputEditorModule.cpp b/Engine/Plugins/EnhancedInput/Source/InputEditor/Private/InputEditorModule.cpp
--- a/Engine/Plugins/EnhancedInput/Source/InputEditor/Private/InputEditorModule.cpp
+++ b/Engine/Plugins/EnhancedInput/Source/InputEditor/Private/InputEditorModule.cpp
@@ -217,7 +217,7 @@ UObject* UInputAction_Factory::FactoryCreateNew(UClass* Class, UObject* InParent
 
 PRAGMA_DISABLE_DEPRECATION_WARNINGS
 // UPlayerMappableInputConfig_Factory
-UPlayerMappableInputConfig_Factory::UPlayerMappableInputConfig_Factory(const class FObjectInitializer& OBJ)
+UDEPRECATED_UPlayerMappableInputConfig_Factory::UDEPRECATED_UPlayerMappableInputConfig_Factory(const class FObjectInitializer& OBJ)
 	: Super(OBJ)
 {
 	SupportedClass = UPlayerMappableInputConfig::StaticClass();
@@ -225,7 +225,7 @@ UPlayerMappableInputConfig_Factory::UPlayerMappableInputConfig_Factory(const cla
 	bCreateNew = true;
 }
 
-UObject* UPlayerMappableInputConfig_Factory::FactoryCreateNew(UClass* Class, UObject* InParent, FName Name, EObjectFlags Flags, UObject* Context, FFeedbackContext* Warn)
+UObject* UDEPRECATED_UPlayerMappableInputConfig_Factory::FactoryCreateNew(UClass* Class, UObject* InParent, FName Name, EObjectFlags Flags, UObject* Context, FFeedbackContext* Warn)
 {
 	check(Class->IsChildOf(UPlayerMappableInputConfig::StaticClass()));
 	return NewObject<UPlayerMappableInputConfig>(InParent, Class, Name, Flags | RF_Transactional, Context);
diff --git a/Engine/Plugins/EnhancedInput/Source/InputEditor/Public/InputEditorModule.h b/Engine/Plugins/EnhancedInput/Source/InputEditor/Public/InputEditorModule.h
--- a/Engine/Plugins/EnhancedInput/Source/InputEditor/Public/InputEditorModule.h
+++ b/Engine/Plugins/EnhancedInput/Source/InputEditor/Public/InputEditorModule.h
@@ -88,8 +88,8 @@ public:
 	virtual UObject* FactoryCreateNew(UClass* Class, UObject* InParent, FName Name, EObjectFlags Flags, UObject* Context, FFeedbackContext* Warn) override;
 };
 
-UCLASS()
-class UE_DEPRECATED(5.3, "UPlayerMappableInputConfig has been deprecated, please use the UEnhancedInputUserSettings system instead.") INPUTEDITOR_API UPlayerMappableInputConfig_Factory : public UFactory
+UCLASS(deprecated)
+class UE_DEPRECATED(5.3, "UPlayerMappableInputConfig has been deprecated, please use the UEnhancedInputUserSettings system instead.") INPUTEDITOR_API UDEPRECATED_UPlayerMappableInputConfig_Factory : public UFactory
 {
 	GENERATED_UCLASS_BODY()
 public:
diff --git a/Engine/Source/Developer/AssetTools/Private/AssetTools.cpp b/Engine/Source/Developer/AssetTools/Private/AssetTools.cpp
--- a/Engine/Source/Developer/AssetTools/Private/AssetTools.cpp
+++ b/Engine/Source/Developer/AssetTools/Private/AssetTools.cpp
@@ -6086,7 +6086,7 @@ TArray<UFactory*> UAssetToolsImpl::GetNewAssetFactories() const
 	for (TObjectIterator<UClass> It; It; ++It)
 	{
 		UClass* Class = *It;
-		if (!Class->IsChildOf(UFactory::StaticClass()) || Class->HasAnyClassFlags(CLASS_Abstract))
+		if (!Class->IsChildOf(UFactory::StaticClass()) || Class->HasAnyClassFlags(CLASS_Abstract) || Class->HasAnyClassFlags(CLASS_Deprecated))
 		{
 			continue;
 		}
