[Input] Exposing trigger state from FInputActionInstance

diff --git a/Engine/Plugins/EnhancedInput/Source/EnhancedInput/Private/InputAction.cpp b/Engine/Plugins/EnhancedInput/Source/EnhancedInput/Private/InputAction.cpp
--- a/Engine/Plugins/EnhancedInput/Source/EnhancedInput/Private/InputAction.cpp
+++ b/Engine/Plugins/EnhancedInput/Source/EnhancedInput/Private/InputAction.cpp
@@ -64,6 +64,11 @@ FInputActionValue FInputActionInstance::GetValue() const
 		FInputActionValue(Value.GetValueType(), FInputActionValue::Axis3D::ZeroVector);
  } 
 
+ETriggerState FInputActionInstance::GetEvaluatedActionTriggerState() const
+{
+	return TriggerStateTracker.GetState();
+}
+
 // Calculate a collective representation of trigger state from evaluations of all triggers in one or more trigger groups.
 ETriggerState FTriggerStateTracker::EvaluateTriggers(const UEnhancedPlayerInput* PlayerInput, const TArray<UInputTrigger*>& Triggers, FInputActionValue ModifiedValue, float DeltaTime)
 {
diff --git a/Engine/Plugins/EnhancedInput/Source/EnhancedInput/Private/InputTriggers.cpp b/Engine/Plugins/EnhancedInput/Source/EnhancedInput/Private/InputTriggers.cpp
--- a/Engine/Plugins/EnhancedInput/Source/EnhancedInput/Private/InputTriggers.cpp
+++ b/Engine/Plugins/EnhancedInput/Source/EnhancedInput/Private/InputTriggers.cpp
@@ -339,7 +339,7 @@ ETriggerState UInputTriggerChordAction::UpdateState_Implementation(const UEnhanc
 {
 	// Inherit state from the chorded action
 	const FInputActionInstance* EventData = PlayerInput->FindActionInstanceData(ChordAction);
-	return EventData ? EventData->TriggerStateTracker.GetState() : ETriggerState::None;
+	return EventData ? EventData->GetEvaluatedActionTriggerState() : ETriggerState::None;
 }
 
 UInputTriggerCombo::UInputTriggerCombo()
diff --git a/Engine/Plugins/EnhancedInput/Source/EnhancedInput/Public/InputAction.h b/Engine/Plugins/EnhancedInput/Source/EnhancedInput/Public/InputAction.h
--- a/Engine/Plugins/EnhancedInput/Source/EnhancedInput/Public/InputAction.h
+++ b/Engine/Plugins/EnhancedInput/Source/EnhancedInput/Public/InputAction.h
@@ -197,7 +197,6 @@ USTRUCT(BlueprintType)
 struct FInputActionInstance
 {
 	friend class UEnhancedPlayerInput;
-	friend class UInputTriggerChordAction;
 
 	GENERATED_BODY()
 
@@ -249,6 +248,9 @@ public:
 	// Current action value - Will be zero if the current trigger event is not ETriggerEvent::Triggered!
 	UE_API FInputActionValue GetValue() const;
 
+	// Returns the evaluated trigger state
+	ETriggerState GetEvaluatedActionTriggerState() const;
+
 	// Total time the action has been evaluating triggering (Ongoing & Triggered)
 	float GetElapsedTime() const { return ElapsedProcessedTime; }
 
