Allow EnhancedInput DEV_ONLY_KEY_BINDINGS_AVAILABLE to be defined by project Target.cs

diff --git a/Engine/Plugins/EnhancedInput/Source/EnhancedInput/Public/EnhancedInputComponent.h b/Engine/Plugins/EnhancedInput/Source/EnhancedInput/Public/EnhancedInputComponent.h
--- a/Engine/Plugins/EnhancedInput/Source/EnhancedInput/Public/EnhancedInputComponent.h
+++ b/Engine/Plugins/EnhancedInput/Source/EnhancedInput/Public/EnhancedInputComponent.h
@@ -10,7 +10,9 @@
 #define UE_API ENHANCEDINPUT_API
 
 // TODO: Disable in test too?
-#define DEV_ONLY_KEY_BINDINGS_AVAILABLE (!UE_BUILD_SHIPPING)
+#ifndef DEV_ONLY_KEY_BINDINGS_AVAILABLE
+	#define DEV_ONLY_KEY_BINDINGS_AVAILABLE (!UE_BUILD_SHIPPING)
+#endif 
 
 enum class ETriggerEvent : uint8;
 
