PR #13138: [Enhanced Input] Improvements to profile handling behavior.

Make sure that the key mappings "supported key profile" tag is respected, and make the OnKeyProfileChanged delegate BP assignable.

diff --git a/Engine/Plugins/EnhancedInput/Source/EnhancedInput/Private/UserSettings/EnhancedInputUserSettings.cpp b/Engine/Plugins/EnhancedInput/Source/EnhancedInput/Private/UserSettings/EnhancedInputUserSettings.cpp
--- a/Engine/Plugins/EnhancedInput/Source/EnhancedInput/Private/UserSettings/EnhancedInputUserSettings.cpp
+++ b/Engine/Plugins/EnhancedInput/Source/EnhancedInput/Private/UserSettings/EnhancedInputUserSettings.cpp
@@ -1592,6 +1592,15 @@ bool UEnhancedInputUserSettings::RegisterKeyMappingsToProfile(UEnhancedPlayerMap
 			continue;
 		}
 
+		// Only mappings with empty profile or the current profile should be added - same as it's done by IEnhancedInputSubsystemInterface::RebuildControlMappings
+		const UPlayerMappableKeySettings* KeySettings = KeyMapping.GetPlayerMappableKeySettings();
+		
+		// If this mapping has specified a specific key profile, and the current profile isn't it, then don't add this key mapping
+		if (KeySettings && !KeySettings->SupportedKeyProfileIds.IsEmpty() && !KeySettings->SupportedKeyProfileIds.Contains(Profile.GetProfileIdString()))
+		{
+			continue;
+		}
+
 		// Get the unique FName for the "Mapping name". This is set on the UInputAction
 		// but can be overridden on each FEnhancedActionKeyMapping to create multiple
 		// mapping options for a single Input Action. 
diff --git a/Engine/Plugins/EnhancedInput/Source/EnhancedInput/Public/UserSettings/EnhancedInputUserSettings.h b/Engine/Plugins/EnhancedInput/Source/EnhancedInput/Public/UserSettings/EnhancedInputUserSettings.h
--- a/Engine/Plugins/EnhancedInput/Source/EnhancedInput/Public/UserSettings/EnhancedInputUserSettings.h
+++ b/Engine/Plugins/EnhancedInput/Source/EnhancedInput/Public/UserSettings/EnhancedInputUserSettings.h
@@ -679,6 +679,7 @@ public:
 
 	/** Fired when you equip a different key profile  */
 	DECLARE_DYNAMIC_MULTICAST_DELEGATE_OneParam(FMappableKeyProfileChanged, const UEnhancedPlayerMappableKeyProfile*, NewProfile);
+	UPROPERTY(BlueprintAssignable, Transient, Category = "Enhanced Input|User Settings")
 	FMappableKeyProfileChanged OnKeyProfileChanged;
 	
 	/**
