MinimalAPI Fix for UInputTriggerTimedBase

diff --git a/Engine/Plugins/EnhancedInput/Source/EnhancedInput/Public/InputTriggers.h b/Engine/Plugins/EnhancedInput/Source/EnhancedInput/Public/InputTriggers.h
--- a/Engine/Plugins/EnhancedInput/Source/EnhancedInput/Public/InputTriggers.h
+++ b/Engine/Plugins/EnhancedInput/Source/EnhancedInput/Public/InputTriggers.h
@@ -184,7 +184,7 @@ Base class for building triggers that have firing conditions governed by elapsed
 This class transitions state to Ongoing once input is actuated, and will track Ongoing input time until input is released.
 Inheriting classes should provide the logic for Triggered transitions.
 */
-UCLASS(Abstract, Config = Input)	// TODO: Parent CLASS_Config flag is passed down to all UInputTrigger descendants, but ClassConfigName is only inherited by immediate children causing UHT to complain about a missing config file for any triggers based on this.
+UCLASS(MinimalApi, Abstract, Config = Input)	// TODO: Parent CLASS_Config flag is passed down to all UInputTrigger descendants, but ClassConfigName is only inherited by immediate children causing UHT to complain about a missing config file for any triggers based on this.
 class UInputTriggerTimedBase : public UInputTrigger
 {
 	GENERATED_BODY()
