[EnhancedInputSubsystem] Allow users to determine if given input action is being continuously injected or not

-Add function that returns if provided input action is being continuously injected.

diff --git a/Engine/Plugins/EnhancedInput/Source/EnhancedInput/Private/EnhancedInputSubsystemInterface.cpp b/Engine/Plugins/EnhancedInput/Source/EnhancedInput/Private/EnhancedInputSubsystemInterface.cpp
--- a/Engine/Plugins/EnhancedInput/Source/EnhancedInput/Private/EnhancedInputSubsystemInterface.cpp
+++ b/Engine/Plugins/EnhancedInput/Source/EnhancedInput/Private/EnhancedInputSubsystemInterface.cpp
@@ -201,6 +201,11 @@ void IEnhancedInputSubsystemInterface::StopContinuousInputInjectionForPlayerMapp
 	}
 }
 
+bool IEnhancedInputSubsystemInterface::HasContinuousInputInjectionForAction(const UInputAction* Action) const
+{
+	return const_cast<IEnhancedInputSubsystemInterface*>(this)->GetContinuouslyInjectedInputs().Contains(Action);
+}
+
 void IEnhancedInputSubsystemInterface::InjectInputVectorForPlayerMapping(const FName MappingName, FVector Value, const TArray<UInputModifier*>& Modifiers, const TArray<UInputTrigger*>& Triggers)
 {
 	if (const UEnhancedInputUserSettings* UserSettings = GetUserSettings())
diff --git a/Engine/Plugins/EnhancedInput/Source/EnhancedInput/Public/EnhancedInputSubsystemInterface.h b/Engine/Plugins/EnhancedInput/Source/EnhancedInput/Public/EnhancedInputSubsystemInterface.h
--- a/Engine/Plugins/EnhancedInput/Source/EnhancedInput/Public/EnhancedInputSubsystemInterface.h
+++ b/Engine/Plugins/EnhancedInput/Source/EnhancedInput/Public/EnhancedInputSubsystemInterface.h
@@ -240,6 +240,11 @@ public:
 	UFUNCTION(BlueprintCallable, Category="Input")
 	UE_API virtual void StopContinuousInputInjectionForPlayerMapping(UPARAM(Meta=(GetOptions="EnhancedInput.PlayerMappableKeySettings.GetKnownMappingNames")) const FName MappingName);
 	
+	/**
+	 * Determines if there is a continuous input injection for the given action.
+	 */
+	UE_API bool HasContinuousInputInjectionForAction(const UInputAction* Action) const;
+	
 	/**
 	 * Remove all applied mapping contexts.
 	 */
