[Enhanced Input] Add a "MappingProfileOverrides" property to UInputMappingContext which allows you to specify different default key mappings for a mapping context for different player mappable key profiles.

When you add an entry into the "MappingProfileOverrides" map, those mappings will be used instead of the default mappings when a profile with that ID is active in the user settings.

For example, you could create a "southpaw" preset which applies inversion modifiers on the default key mappings, and expose it as a "Preset" in the settings screen.

- Encapsulates the TArray<FEnhancedActionKeyMapping> of key mappings into a data struct so that we can re-use it's details panel customizations
- Bump serialization version to auto-fixup any older assets and populate the new struct data in PostLoad
- Add a command to set the current key profile by name, "EnhancedInput.SetActiveKeyProfile"

[RN] Allow users to author new default key values for specific player mappable key profiles directly within an Input Mapping Context asset.

diff --git a/Engine/Plugins/EnhancedInput/Config/BaseEnhancedInput.ini b/Engine/Plugins/EnhancedInput/Config/BaseEnhancedInput.ini
--- a/Engine/Plugins/EnhancedInput/Config/BaseEnhancedInput.ini
+++ b/Engine/Plugins/EnhancedInput/Config/BaseEnhancedInput.ini
@@ -13,4 +13,10 @@
 +FunctionRedirects=(OldName="/Script/EnhancedInput.EnhancedInputUserSettings.GetCurrentKeyProfileIdentifier", NewName="/Script/EnhancedInput.EnhancedInputUserSettings.GetActiveKeyProfileId")
 +FunctionRedirects=(OldName="/Script/EnhancedInput.EnhancedPlayerMappableKeyProfile.GetProfileIdentifer", NewName="/Script/EnhancedInput.EnhancedPlayerMappableKeyProfile.GetProfileIdString")
 +PropertyRedirects=(OldName="EnhancedInput.MapPlayerKeyArgs.ProfileId",NewName="EnhancedInput.MapPlayerKeyArgs.ProfileIdString")
-+PropertyRedirects=(OldName="EnhancedInput.PlayerMappableKeyProfileCreationArgs.ProfileIdentifier",NewName="EnhancedInput.PlayerMappableKeyProfileCreationArgs.ProfileStringIdentifier")
\ No newline at end of file
++PropertyRedirects=(OldName="EnhancedInput.PlayerMappableKeyProfileCreationArgs.ProfileIdentifier",NewName="EnhancedInput.PlayerMappableKeyProfileCreationArgs.ProfileStringIdentifier")
+
+;5.7
++PropertyRedirects=(OldName="EnhancedInput.InputMappingContext.Mappings",NewName="EnhancedInput.InputMappingContext.DefaultKeyMappings.Mappings")
+
+;[/Script/Engine.Engine]
+;+ActiveClassRedirects=(OldName="/Script/EnhancedInput.InputMappingContext:Mappings", NewName="/Script/EnhancedInput.InputMappingContext:DefaultKeyMappings.Mappings")
diff --git a/Engine/Plugins/EnhancedInput/Source/EnhancedInput/Private/EnhancedInputSubsystemInterface.cpp b/Engine/Plugins/EnhancedInput/Source/EnhancedInput/Private/EnhancedInputSubsystemInterface.cpp
--- a/Engine/Plugins/EnhancedInput/Source/EnhancedInput/Private/EnhancedInputSubsystemInterface.cpp
+++ b/Engine/Plugins/EnhancedInput/Source/EnhancedInput/Private/EnhancedInputSubsystemInterface.cpp
@@ -469,6 +469,10 @@ EMappingQueryResult IEnhancedInputSubsystemInterface::QueryMapKeyInContextSet(co
 
 	EMappingQueryResult Result = EMappingQueryResult::MappingAvailable;
 
+	UEnhancedInputUserSettings* CurrentUserSettings = GetUserSettings();
+	UEnhancedPlayerMappableKeyProfile* PlayerKeyProfile = CurrentUserSettings ? CurrentUserSettings->GetActiveKeyProfile() : nullptr;
+	const FString& CurrentProfileId = PlayerKeyProfile ? PlayerKeyProfile->GetProfileIdString() : TEXT("");
+
 	// These will be ordered by priority
 	for (const UInputMappingContext* BlockingContext : PrioritizedActiveContexts)
 	{
@@ -487,7 +491,7 @@ EMappingQueryResult IEnhancedInputSubsystemInterface::QueryMapKeyInContextSet(co
 			Stage = EStage::Main;
 		}
 
-		for (const FEnhancedActionKeyMapping& Mapping : BlockingContext->GetMappings())
+		for (const FEnhancedActionKeyMapping& Mapping : BlockingContext->GetMappingsForProfile(CurrentProfileId))
 		{
 			if (Mapping.Key == Key && Mapping.Action)
 			{
@@ -953,7 +957,7 @@ void IEnhancedInputSubsystemInterface::RebuildControlMappings()
 			continue;
 		}
 		
-		TArray<FEnhancedActionKeyMapping> OrderedMappings = ReorderMappings(MappingContext->GetMappings(), PlayerInput->DependentChordActions);
+		TArray<FEnhancedActionKeyMapping> OrderedMappings = ReorderMappings(MappingContext->GetMappingsForProfile(PlayerKeyProfile ? PlayerKeyProfile->GetProfileIdString() : TEXT("")), PlayerInput->DependentChordActions);
 
 		for (FEnhancedActionKeyMapping& Mapping : OrderedMappings)
 		{
diff --git a/Engine/Plugins/EnhancedInput/Source/EnhancedInput/Private/InputMappingContext.cpp b/Engine/Plugins/EnhancedInput/Source/EnhancedInput/Private/InputMappingContext.cpp
--- a/Engine/Plugins/EnhancedInput/Source/EnhancedInput/Private/InputMappingContext.cpp
+++ b/Engine/Plugins/EnhancedInput/Source/EnhancedInput/Private/InputMappingContext.cpp
@@ -7,6 +7,7 @@
 #include "EnhancedInputModule.h"
 #include "Misc/DataValidation.h"
 #include "PlayerMappableKeySettings.h"
+#include "UObject/FortniteSeasonBranchObjectVersion.h"
 
 #include UE_INLINE_GENERATED_CPP_BY_NAME(InputMappingContext)
 
@@ -16,14 +17,55 @@
 EDataValidationResult UInputMappingContext::IsDataValid(FDataValidationContext& Context) const
 {
 	EDataValidationResult Result = CombineDataValidationResults(Super::IsDataValid(Context), EDataValidationResult::Valid);
-	for (const FEnhancedActionKeyMapping& Mapping : Mappings)
+
+	auto ValidateMapping = [&Result, &Context](const FEnhancedActionKeyMapping& Mapping)
 	{
 		Result = CombineDataValidationResults(Result, Mapping.IsDataValid(Context));
-	}
+	};
+
+	// Run validation for every key mapping
+	ForEachKeyMapping(ValidateMapping);
+	
 	return Result;
 }
 #endif	// WITH_EDITOR
 
+void UInputMappingContext::Serialize(FArchive& Ar)
+{
+	Super::Serialize(Ar);
+
+	Ar.UsingCustomVersion(FFortniteSeasonBranchObjectVersion::GUID);
+}
+
+void UInputMappingContext::PostLoad()
+{
+	Super::PostLoad();
+
+PRAGMA_DISABLE_DEPRECATION_WARNINGS
+	if (GetLinkerCustomVersion(FFortniteSeasonBranchObjectVersion::GUID) < FFortniteSeasonBranchObjectVersion::EnhancedInputMappingContextProfileMappingsUpdate)
+	{
+		// We converted from a TArray<FEnhancedActionKeyMapping> to a struct so that it is consistent with the profile override types
+		DefaultKeyMappings.Mappings = MoveTemp(Mappings);
+	}
+PRAGMA_ENABLE_DEPRECATION_WARNINGS
+}
+
+void UInputMappingContext::ForEachKeyMapping(TFunctionRef<void(const FEnhancedActionKeyMapping&)>&& Func) const
+{
+	for (const FEnhancedActionKeyMapping& Mapping : DefaultKeyMappings.Mappings)
+	{
+		Func(Mapping);
+	}
+	
+	for (const TPair<FString, FInputMappingContextMappingData>& ProfileMapping : MappingProfileOverrides)
+	{
+		for (const FEnhancedActionKeyMapping& Mapping : ProfileMapping.Value.Mappings)
+		{
+			Func(Mapping);
+		}
+	}
+}
+
 bool UInputMappingContext::ShouldFilterMappingByInputMode() const
 {
 	return InputModeFilterOptions != EMappingContextInputModeFilterOptions::DoNotFilter;
@@ -55,25 +97,76 @@ bool UInputMappingContext::ShouldShowInputModeQuery()
 	return GetDefault<UEnhancedInputDeveloperSettings>()->bEnableInputModeFiltering;
 }
 
+bool UInputMappingContext::HasMappingsForProfile(const FString& ProfileId) const
+{
+	return MappingProfileOverrides.Find(ProfileId) != nullptr;
+}
+
+TArray<FString> UInputMappingContext::GetProfilesWithOverridenMappings() const
+{
+	TArray<FString> Profiles;
+	MappingProfileOverrides.GetKeys(Profiles);
+	
+	return Profiles;
+}
+
+const TArray<FEnhancedActionKeyMapping>& UInputMappingContext::GetMappingsForProfile(const FString& ProfileId) const
+{
+	if (const FInputMappingContextMappingData* OverrideMappings = MappingProfileOverrides.Find(ProfileId))
+	{
+		return OverrideMappings->Mappings;
+	}
+	
+	return DefaultKeyMappings.Mappings;
+}
+
+bool UInputMappingContext::HasMappingForInputAction(const UInputAction* Action) const
+{
+	if (!Action)
+	{
+		return false;
+	}
+
+	auto HasReferencesToMapping = [Action](const FEnhancedActionKeyMapping& Mapping) -> bool
+	{
+		return Mapping.Action == Action;
+	};
+	
+	if (DefaultKeyMappings.Mappings.ContainsByPredicate(HasReferencesToMapping))
+	{
+		return true;
+	}
+
+	for (const TPair<FString, FInputMappingContextMappingData>& OverrideMappings : MappingProfileOverrides)
+	{
+		if (OverrideMappings.Value.Mappings.ContainsByPredicate(HasReferencesToMapping))
+		{
+			return true;
+		}
+	}
+
+	return false;
+}
+
 FEnhancedActionKeyMapping& UInputMappingContext::MapKey(const UInputAction* Action, FKey ToKey)
 {
 	IEnhancedInputModule::Get().GetLibrary()->RequestRebuildControlMappingsUsingContext(this);
-	return Mappings.Add_GetRef(FEnhancedActionKeyMapping(Action, ToKey));
+	return DefaultKeyMappings.Mappings.Add_GetRef(FEnhancedActionKeyMapping(Action, ToKey));
 }
 
 void UInputMappingContext::UnmapKey(const UInputAction* Action, FKey Key)
 {
-	int32 MappingIdx = Mappings.IndexOfByPredicate([&Action, &Key](const FEnhancedActionKeyMapping& Other) { return Other.Action == Action && Other.Key == Key; });
+	int32 MappingIdx = DefaultKeyMappings.Mappings.IndexOfByPredicate([&Action, &Key](const FEnhancedActionKeyMapping& Other) { return Other.Action == Action && Other.Key == Key; });
 	if (MappingIdx != INDEX_NONE)
 	{
-		Mappings.RemoveAtSwap(MappingIdx);
+		DefaultKeyMappings.Mappings.RemoveAtSwap(MappingIdx);
 		IEnhancedInputModule::Get().GetLibrary()->RequestRebuildControlMappingsUsingContext(this);
 	}
 }
 
 void UInputMappingContext::UnmapAllKeysFromAction(const UInputAction* Action)
 {
-	int32 Found = Mappings.RemoveAllSwap([&Action](const FEnhancedActionKeyMapping& Mapping) { return Mapping.Action == Action; });
+	int32 Found = DefaultKeyMappings.Mappings.RemoveAllSwap([&Action](const FEnhancedActionKeyMapping& Mapping) { return Mapping.Action == Action; });
 	if (Found > 0)
 	{
 		IEnhancedInputModule::Get().GetLibrary()->RequestRebuildControlMappingsUsingContext(this);
@@ -82,9 +175,9 @@ void UInputMappingContext::UnmapAllKeysFromAction(const UInputAction* Action)
 
 void UInputMappingContext::UnmapAll()
 {
-	if (Mappings.Num())
+	if (DefaultKeyMappings.Mappings.Num())
 	{
-		Mappings.Empty();
+		DefaultKeyMappings.Mappings.Empty();
 		IEnhancedInputModule::Get().GetLibrary()->RequestRebuildControlMappingsUsingContext(this);
 	}
 }
diff --git a/Engine/Plugins/EnhancedInput/Source/EnhancedInput/Private/UserSettings/EnhancedInputUserSettings.cpp b/Engine/Plugins/EnhancedInput/Source/EnhancedInput/Private/UserSettings/EnhancedInputUserSettings.cpp
--- a/Engine/Plugins/EnhancedInput/Source/EnhancedInput/Private/UserSettings/EnhancedInputUserSettings.cpp
+++ b/Engine/Plugins/EnhancedInput/Source/EnhancedInput/Private/UserSettings/EnhancedInputUserSettings.cpp
@@ -68,6 +68,25 @@ namespace UE::EnhancedInput
 		});
 	}
 	
+	static void SetActiveKeyProfile(const TArray<FString>& Args)
+	{
+		if (!Args.IsValidIndex(0))
+		{
+			UE_LOG(LogEnhancedInput, Error, TEXT("EnhancedInput.SetActiveKeyProfile command failed - you must provide a desired key profile name! Example usage:\n\t EnhancedInput.SetActiveKeyProfile PresetA"));
+			return;
+		}
+
+		const FString& DesiredKeyProfile = Args[0];
+		
+		UEnhancedInputLibrary::ForEachSubsystem([DesiredKeyProfile](IEnhancedInputSubsystemInterface* Subsystem)
+		{
+			if (UEnhancedInputUserSettings* Settings = Subsystem->GetUserSettings())
+			{
+				Settings->SetActiveKeyProfile(DesiredKeyProfile);
+			}
+		});
+	}
+	
 	static FAutoConsoleCommand ConsoleCommandDumpProfileToLog(
 		TEXT("EnhancedInput.DumpKeyProfileToLog"),
 		TEXT(""),
@@ -77,6 +96,11 @@ namespace UE::EnhancedInput
 		TEXT("EnhancedInput.SaveKeyProfilesToSlot"),
 		TEXT("Save the user input settings object with the Save Game to slot system"),
 		FConsoleCommandWithArgsDelegate::CreateStatic(UE::EnhancedInput::SaveAllKeyProfilesToSlot));
+
+	static FAutoConsoleCommand ConsoleCommandSetKeyProfile(
+		TEXT("EnhancedInput.SetActiveKeyProfile"),
+		TEXT("Sets the current Enhanced Input user key profile to the one provided."),
+		FConsoleCommandWithArgsDelegate::CreateStatic(UE::EnhancedInput::SetActiveKeyProfile));
 }
 
 ///////////////////////////////////////////////////////////
@@ -1557,6 +1581,26 @@ bool UEnhancedInputUserSettings::RegisterInputMappingContext(const UInputMapping
 
 	bool bResult = true;
 
+	const ULocalPlayer* LP = GetLocalPlayer();
+
+	// Create any key profiles for mapping contexts which have specific overrides if we need to
+	const TArray<FString> ProfilesSpecifiedInMappingContext = IMC->GetProfilesWithOverridenMappings();
+	for (const FString& ProfileId : ProfilesSpecifiedInMappingContext)
+	{
+		if (SavedKeyProfilesMap.Contains(ProfileId))
+		{
+			continue;
+		}
+
+		FPlayerMappableKeyProfileCreationArgs Args = {};
+		Args.ProfileStringIdentifier = ProfileId;
+		Args.UserId = LP ? LP->GetPlatformUserId() : PLATFORMUSERID_NONE;
+		Args.DisplayName = UE::EnhancedInput::DefaultProfileDisplayName;
+		Args.bSetAsCurrentProfile = false;
+		
+		CreateNewKeyProfile(Args);
+	}
+
 	// Register the mappings of this IMC to every saved key profile
 	for (TPair<FString, TObjectPtr<UEnhancedPlayerMappableKeyProfile>>& Pair : SavedKeyProfilesMap)
 	{
@@ -1584,7 +1628,7 @@ bool UEnhancedInputUserSettings::RegisterKeyMappingsToProfile(UEnhancedPlayerMap
 	static TMap<FHardwareDeviceIdentifier, FSlotCountToMappingNames> MappingsPerDeviceTypes;
 	MappingsPerDeviceTypes.Reset();
 	
-	for (const FEnhancedActionKeyMapping& KeyMapping : IMC->GetMappings())
+	for (const FEnhancedActionKeyMapping& KeyMapping : IMC->GetMappingsForProfile(Profile.GetProfileIdString()))
 	{
 		// Skip over non-player mappable keys
 		if (!KeyMapping.IsPlayerMappable())
diff --git a/Engine/Plugins/EnhancedInput/Source/EnhancedInput/Public/InputMappingContext.h b/Engine/Plugins/EnhancedInput/Source/EnhancedInput/Public/InputMappingContext.h
--- a/Engine/Plugins/EnhancedInput/Source/EnhancedInput/Public/InputMappingContext.h
+++ b/Engine/Plugins/EnhancedInput/Source/EnhancedInput/Public/InputMappingContext.h
@@ -63,6 +63,18 @@ enum class EMappingContextRegistrationTrackingMode : uint8
 	CountRegistrations
 };
 
+USTRUCT(BlueprintType)
+struct FInputMappingContextMappingData
+{
+	GENERATED_BODY()
+	
+	/**
+	 * List of Input Actions and which keys they are mapped to.
+	 */
+	UPROPERTY(config, BlueprintReadOnly, EditAnywhere, Category = "Mappings")
+	TArray<FEnhancedActionKeyMapping> Mappings;
+};
+
 /**
 * UInputMappingContext : A collection of key to action mappings for a specific input context
 * Could be used to:
@@ -78,15 +90,34 @@ class UInputMappingContext : public UDataAsset
 
 protected:
 	// List of key to action mappings.
-	UPROPERTY(config, BlueprintReadOnly, EditAnywhere, Category = "Mappings")
+	UE_DEPRECATED(5.7, "Use the DefaultKeyMappings struct instead.")
+	UPROPERTY(config, BlueprintReadOnly, Category = "Mappings", meta = (DeprecatedProperty="Note", DeprecationMessage = "Use the DefaultKeyMappings struct instead."))
 	TArray<FEnhancedActionKeyMapping> Mappings;
 
+	/**
+	 * List of Input Actions to which FKeys they are mapped to.
+	 */
+	UPROPERTY(config, BlueprintReadOnly, EditAnywhere, Category = "Mappings")
+	FInputMappingContextMappingData DefaultKeyMappings;
+	
+	/**
+	 * Key mappings that should be used INSTEAD of the default "Mappings" array above then the
+	 * provided Remappable Key Profile of the given name is currently active.
+	 *
+	 * If there are no override mappings specified for a set for a specific mapping profile, then the default mappings will be used. 
+	 */
+	UPROPERTY(config, BlueprintReadOnly, EditAnywhere, Category = "Mappings")
+	TMap<FString, FInputMappingContextMappingData> MappingProfileOverrides; 
+
 public:
 
 #if WITH_EDITOR
 	UE_API virtual EDataValidationResult IsDataValid(class FDataValidationContext& Context) const override;
 #endif
-
+	
+	UE_API virtual void Serialize(FArchive& Ar) override;
+	UE_API virtual void PostLoad() override;
+	
 protected:
 	
 	/**
@@ -116,6 +147,13 @@ protected:
 	UPROPERTY(BlueprintReadOnly, EditAnywhere, Category = "Registration")
 	EMappingContextRegistrationTrackingMode RegistrationTrackingMode = EMappingContextRegistrationTrackingMode::Untracked;
 
+	/**
+	 * Executes the given lambda on every default and override key mapping specified on this mapping context.
+	 * 
+	 * @param Func Lambda function to execute on a given key mapping
+	 */
+	UE_API void ForEachKeyMapping(TFunctionRef<void(const FEnhancedActionKeyMapping&)>&& Func) const;
+
 public:
 
 	/**
@@ -124,7 +162,7 @@ public:
 	UE_API bool ShouldFilterMappingByInputMode() const;
 
 	/**
-	 * @return The tag query which should be used when deciding whether this mapping context should be filterd
+	 * @return The tag query which should be used when deciding whether this mapping context should be filtered
 	 * out based on the current input mode or not. 
 	 */
 	UE_API FGameplayTagQuery GetInputModeQuery() const;
@@ -145,12 +183,38 @@ public:
 	UFUNCTION()
 	static UE_API bool ShouldShowInputModeQuery();
 
+	/**
+	 * @param ProfileId The Player Mappable Key Profile ID
+	 * @return True if this Input Mapping Context has some specific override mappings for the given profile. 
+	 */
+	bool HasMappingsForProfile(const FString& ProfileId) const;
+
+	/**
+	 * @return An array of profile identifiers which this IMC specifies overrides for.
+	 */
+	TArray<FString> GetProfilesWithOverridenMappings() const;
+
+	/**
+	 * @param ProfileId The Player Mappable Key Profile ID
+	 * @return Reference to the key mappings to use for the given profile. If there is no override mapping specified for the given
+	 *		   profile, then the default mappings will be returned.
+	 */
+	const TArray<FEnhancedActionKeyMapping>& GetMappingsForProfile(const FString& ProfileId) const;
+
+	/**
+	 * Returns true if this Input Mapping Context contains a reference to the given Input Action in its default or profile override mappings.
+	 */
+	bool HasMappingForInputAction(const UInputAction* Action) const;
+
+	
+	// TODO_BH: Deprecate all of these functions below. They are only used in the editor and for tests! Then we can also get rid of the friend decls above
+	
 	/**
 	* Mapping accessors.
 	* Note: Use UEnhancedInputLibrary::RequestRebuildControlMappingsForContext to invoke changes made to an FEnhancedActionKeyMapping
 	*/
-	const TArray<FEnhancedActionKeyMapping>& GetMappings() const { return Mappings; }
-	FEnhancedActionKeyMapping& GetMapping(TArray<FEnhancedActionKeyMapping>::SizeType Index) { return Mappings[Index]; }
+	const TArray<FEnhancedActionKeyMapping>& GetMappings() const { return DefaultKeyMappings.Mappings; }
+	FEnhancedActionKeyMapping& GetMapping(TArray<FEnhancedActionKeyMapping>::SizeType Index) { return DefaultKeyMappings.Mappings[Index]; }
 
 	// TODO: Don't want to encourage Map/Unmap calls here where context switches would be desirable. These are intended for use in the config/binding screen only.
 
@@ -179,8 +243,4 @@ public:
 	UE_API void UnmapAll();
 };
 
-// ************************************************************************************************
-// ************************************************************************************************
-// ************************************************************************************************
-
 #undef UE_API
diff --git a/Engine/Plugins/EnhancedInput/Source/InputEditor/Private/ActionMappingDetails.cpp b/Engine/Plugins/EnhancedInput/Source/InputEditor/Private/ActionMappingDetails.cpp
--- a/Engine/Plugins/EnhancedInput/Source/InputEditor/Private/ActionMappingDetails.cpp
+++ b/Engine/Plugins/EnhancedInput/Source/InputEditor/Private/ActionMappingDetails.cpp
@@ -19,9 +19,8 @@
 
 // TODO: This is derived from (and will eventually replace) InputSettingsDetails.cpp
 
-FActionMappingsNodeBuilderEx::FActionMappingsNodeBuilderEx(IDetailLayoutBuilder* InDetailLayoutBuilder, const TSharedPtr<IPropertyHandle>& InPropertyHandle)
-	: DetailLayoutBuilder(InDetailLayoutBuilder)
-	, ActionMappingsPropertyHandle(InPropertyHandle)
+FActionMappingsNodeBuilderEx::FActionMappingsNodeBuilderEx(const TSharedPtr<IPropertyHandle>& InPropertyHandle)
+	: ActionMappingsPropertyHandle(InPropertyHandle)
 {
 	// Support for updating references to renamed input actions
 	FAssetRegistryModule& AssetRegistryModule = FModuleManager::LoadModuleChecked<FAssetRegistryModule>(FName("AssetRegistry"));
@@ -89,6 +88,8 @@ void FActionMappingsNodeBuilderEx::ReorderMappings(int32 OriginalIndex, int32 Ne
 	ActionMappingsPropertyHandle->NotifyPreChange();
 
 	DelayedGroupExpansionStates.Emplace(nullptr, true);
+
+	TSharedPtr<IPropertyHandleArray> ActionMappingsArrayHandle = ActionMappingsPropertyHandle->AsArray();
 	
 	if (ReorderMode == EReorderMode::Group)
 	{
@@ -97,7 +98,7 @@ void FActionMappingsNodeBuilderEx::ReorderMappings(int32 OriginalIndex, int32 Ne
 		//1: Get the list of groups in their current order, and assign them an index, with room to insert the moved group between any groups.
 		TMap<const UObject*, uint32> Groups;
 		uint32 GroupIndex = 1;
-		const TSharedPtr<IPropertyHandleArray> ActionMappingsArrayHandle = ActionMappingsPropertyHandle->AsArray();
+		
 		uint32 NumMappings;
 		ActionMappingsArrayHandle->GetNumElements(NumMappings);
 		for (uint32 Index = 0; Index < NumMappings; ++Index)
@@ -121,20 +122,23 @@ void FActionMappingsNodeBuilderEx::ReorderMappings(int32 OriginalIndex, int32 Ne
 		Groups[ActionOfOriginalIndex] = (NewIndex * 2 + ((OriginalIndex < NewIndex) ? 2 : 0));
 
 		//3: Stable sort, to sort group in the new required order, but keep the order inside each category.
-		InputContext->Mappings.StableSort([&Groups](const FEnhancedActionKeyMapping& A, const FEnhancedActionKeyMapping& B) { return Groups[A.Action] < Groups[B.Action]; });
+		TSharedPtr<IPropertyHandle> OuterStructContainerHandle = ActionMappingsPropertyHandle->GetParentHandle();
+		if(OuterStructContainerHandle.IsValid())
+		{
+			void* StructData = nullptr;
+			if(FPropertyAccess::Success == OuterStructContainerHandle->GetValueData(StructData))
+			{
+				FInputMappingContextMappingData* RawData = static_cast<FInputMappingContextMappingData*>(StructData);
+				if (RawData)
+				{
+					RawData->Mappings.StableSort([&Groups](const FEnhancedActionKeyMapping& A, const FEnhancedActionKeyMapping& B) { return Groups[A.Action] < Groups[B.Action]; });
+				}
+			}
+		}
 	}
 	else if (ReorderMode == EReorderMode::Single)
 	{
-		if (FMath::Abs(NewIndex - OriginalIndex) == 1) //If both reordered elements are next to each other, we can avoid the cost of a remove/insert.
-		{
-			InputContext->Mappings.Swap(OriginalIndex, NewIndex);
-		}
-		else
-		{
-			FEnhancedActionKeyMapping EnhancedActionKeyMapping(InputContext->Mappings[OriginalIndex]);
-			InputContext->Mappings.RemoveAt(OriginalIndex);
-			InputContext->Mappings.Insert(EnhancedActionKeyMapping, NewIndex);
-		}
+		ActionMappingsArrayHandle->SwapItems(OriginalIndex, NewIndex);
 	}
 	else
 	{
@@ -534,8 +538,10 @@ void FActionMappingsNodeBuilderEx::AddActionMappingButton_OnClick()
 		InputContext->Modify();
 		ActionMappingsPropertyHandle->NotifyPreChange();
 
-		DelayedGroupExpansionStates.Emplace(nullptr, true);
-		InputContext->MapKey(nullptr, FKey());
+		DelayedGroupExpansionStates.Emplace(nullptr, true);		
+		// Add the new mapping to the array...
+		TSharedPtr<IPropertyHandleArray> ArrayHandle = ActionMappingsPropertyHandle->AsArray();
+		const FPropertyHandleItemAddResult Res = ArrayHandle->AddItem();
 
 		ActionMappingsPropertyHandle->NotifyPostChange(EPropertyChangeType::ArrayAdd);
 	}
@@ -564,23 +570,28 @@ void FActionMappingsNodeBuilderEx::OnActionMappingActionChanged(const FAssetData
 		// This can be jarring for the user if they are working with a complex IMC, so we will put a
 		// small toast up to tell them about it if we can
 		{
-			TArray<UObject*> OuterObjects;
-			ActionMappingsPropertyHandle->GetOuterObjects(OuterObjects);
-			if (!OuterObjects.IsEmpty())
+			bool bHasExistingAction = false;
+				
+			for (int32 Index = 0; Index < MappingSet.Mappings.Num(); ++Index)
 			{
-				if (const UInputMappingContext* InputContext = CastChecked<UInputMappingContext>(OuterObjects[0]))
+				const UObject* Action;
+				MappingSet.Mappings[Index]->GetChildHandle(GET_MEMBER_NAME_CHECKED(FEnhancedActionKeyMapping, Action))->GetValue(Action);
+				if (Action == SelectedAction)
 				{
-					if (InputContext->Mappings.FindByPredicate([&SelectedAction](const FEnhancedActionKeyMapping& Mapping) { return Mapping.Action == SelectedAction; }) != nullptr)
-					{
-						FFormatNamedArguments Args;
-						Args.Add(TEXT("ExistingAction"), FText::FromString(*GetNameSafe(SelectedAction)));
-						const FText NotifText = FText::Format(LOCTEXT("CombiningInputActionMappings_Notif", "'{ExistingAction}' already has mappings in this context, combining them!"), Args);
-						
-						FNotificationInfo Notif(NotifText);
-						Notif.ExpireDuration = 3.0f;
-						FSlateNotificationManager::Get().AddNotification(Notif);						
-					}
-				}	
+					bHasExistingAction = true;
+					break;
+				}
+			}
+			
+			if (bHasExistingAction)
+			{
+				FFormatNamedArguments Args;
+				Args.Add(TEXT("ExistingAction"), FText::FromString(*GetNameSafe(SelectedAction)));
+				const FText NotifText = FText::Format(LOCTEXT("CombiningInputActionMappings_Notif", "'{ExistingAction}' already has mappings in this context, combining them!"), Args);
+				
+				FNotificationInfo Notif(NotifText);
+				Notif.ExpireDuration = 3.0f;
+				FSlateNotificationManager::Get().AddNotification(Notif);						
 			}
 		}
 		
@@ -613,8 +624,18 @@ void FActionMappingsNodeBuilderEx::AddActionMappingToGroupButton_OnClick(const A
 		ActionMappingsPropertyHandle->NotifyPreChange();
 
 		DelayedGroupExpansionStates.Emplace(MappingSet.SharedAction, true);
-		InputContext->MapKey(MappingSet.SharedAction, FKey());
 
+		// Create a new FEnhancedActionKeyMapping entry in the mappings array, and set its default Input Action
+		// value to the UInputAction that we would like.
+		{
+			// Add the item to the array...
+			TSharedPtr<IPropertyHandleArray> ArrayHandle = ActionMappingsPropertyHandle->AsArray();
+			const FPropertyHandleItemAddResult Res = ArrayHandle->AddItem();
+
+			// Initalize the UInputAction property...
+			TSharedPtr<IPropertyHandle> NewMappingInputActionProperty = ArrayHandle->GetElement(Res.GetIndex())->GetChildHandle(GET_MEMBER_NAME_CHECKED(FEnhancedActionKeyMapping, Action));
+			NewMappingInputActionProperty->SetValue(MappingSet.SharedAction);
+		}
 		ActionMappingsPropertyHandle->NotifyPostChange(EPropertyChangeType::ArrayAdd);
 	}
 }
diff --git a/Engine/Plugins/EnhancedInput/Source/InputEditor/Private/ActionMappingDetails.h b/Engine/Plugins/EnhancedInput/Source/InputEditor/Private/ActionMappingDetails.h
--- a/Engine/Plugins/EnhancedInput/Source/InputEditor/Private/ActionMappingDetails.h
+++ b/Engine/Plugins/EnhancedInput/Source/InputEditor/Private/ActionMappingDetails.h
@@ -33,8 +33,8 @@ namespace ActionMappingDetails
 class FActionMappingsNodeBuilderEx : public IDetailCustomNodeBuilder, public TSharedFromThis<FActionMappingsNodeBuilderEx>
 {
 public:
-	FActionMappingsNodeBuilderEx(IDetailLayoutBuilder* InDetailLayoutBuilder, const TSharedPtr<IPropertyHandle>& InPropertyHandle);
-	virtual ~FActionMappingsNodeBuilderEx();
+	FActionMappingsNodeBuilderEx(const TSharedPtr<IPropertyHandle>& InPropertyHandle);
+	virtual ~FActionMappingsNodeBuilderEx() override;
 
 	/** IDetailCustomNodeBuilder interface */
 	virtual void SetOnRebuildChildren(FSimpleDelegate InOnRebuildChildren) override { OnRebuildChildren = InOnRebuildChildren; }
@@ -86,9 +86,6 @@ private:
 	/** Called to rebuild the children of the detail tree */
 	FSimpleDelegate OnRebuildChildren;
 
-	/** Associated detail layout builder */
-	IDetailLayoutBuilder* DetailLayoutBuilder;
-
 	/** Property handle to associated Action Mappings */
 	TSharedPtr<IPropertyHandle> ActionMappingsPropertyHandle;
 
diff --git a/Engine/Plugins/EnhancedInput/Source/InputEditor/Private/InputCustomizations.cpp b/Engine/Plugins/EnhancedInput/Source/InputEditor/Private/InputCustomizations.cpp
--- a/Engine/Plugins/EnhancedInput/Source/InputEditor/Private/InputCustomizations.cpp
+++ b/Engine/Plugins/EnhancedInput/Source/InputEditor/Private/InputCustomizations.cpp
@@ -25,24 +25,6 @@
 
 #define LOCTEXT_NAMESPACE "InputCustomization"
 
-//////////////////////////////////////////////////////////
-// FInputContextDetails
-
-TSharedRef<IDetailCustomization> FInputContextDetails::MakeInstance()
-{
-	return MakeShareable(new FInputContextDetails);
-}
-
-void FInputContextDetails::CustomizeDetails(class IDetailLayoutBuilder& DetailBuilder)
-{
-	// Custom Action Mappings
-	const TSharedPtr<IPropertyHandle> ActionMappingsPropertyHandle = DetailBuilder.GetProperty(GET_MEMBER_NAME_CHECKED(UInputMappingContext, Mappings));
-	ActionMappingsPropertyHandle->MarkHiddenByCustomization();
-
-	IDetailCategoryBuilder& MappingsDetailCategoryBuilder = DetailBuilder.EditCategory(ActionMappingsPropertyHandle->GetDefaultCategoryName());
-	const TSharedRef<FActionMappingsNodeBuilderEx> ActionMappingsBuilder = MakeShareable(new FActionMappingsNodeBuilderEx(&DetailBuilder, ActionMappingsPropertyHandle));
-	MappingsDetailCategoryBuilder.AddCustomBuilder(ActionMappingsBuilder);
-}
 
 //////////////////////////////////////////////////////////
 // FEnhancedActionMappingCustomization
@@ -672,4 +654,46 @@ bool FPlayerMappableKeyChildSettingsCustomization::OnVerifyMappingName(const FTe
 	return Result == EValidatorResult::Ok || Result == EValidatorResult::ExistingName;
 }
 
-#undef LOCTEXT_NAMESPACE
\ No newline at end of file
+/////////////////////////////////////////////////////////////////////////////////////
+// FMappingContextKeyContainerCustomization
+
+TSharedRef<IPropertyTypeCustomization> FMappingContextKeyContainerCustomization::MakeInstance()
+{
+	return MakeShareable(new FMappingContextKeyContainerCustomization());
+}
+
+void FMappingContextKeyContainerCustomization::CustomizeHeader(
+	TSharedRef<IPropertyHandle> PropertyHandle,
+	FDetailWidgetRow& HeaderRow,
+	IPropertyTypeCustomizationUtils& CustomizationUtils)
+{
+	MappingsArrayHandle = PropertyHandle->GetChildHandle(GET_MEMBER_NAME_CHECKED(FInputMappingContextMappingData, Mappings));
+	
+	HeaderRow.NameContent()
+	[
+		PropertyHandle->CreatePropertyNameWidget()
+	];
+	
+	HeaderRow.ValueContent()
+	[
+		PropertyHandle->CreatePropertyValueWidget()
+	];	
+}
+
+void FMappingContextKeyContainerCustomization::CustomizeChildren(
+	TSharedRef<IPropertyHandle> PropertyHandle,
+	IDetailChildrenBuilder& ChildBuilder,
+	IPropertyTypeCustomizationUtils& CustomizationUtils)
+{
+	if (!ensure(MappingsArrayHandle.IsValid()))
+	{
+		return;
+	}
+	
+	MappingsArrayHandle->MarkHiddenByCustomization();
+
+	const TSharedRef<FActionMappingsNodeBuilderEx> ActionMappingsBuilder = MakeShareable(new FActionMappingsNodeBuilderEx(MappingsArrayHandle));
+	ChildBuilder.AddCustomBuilder(ActionMappingsBuilder);
+}
+
+#undef LOCTEXT_NAMESPACE
diff --git a/Engine/Plugins/EnhancedInput/Source/InputEditor/Private/InputCustomizations.h b/Engine/Plugins/EnhancedInput/Source/InputEditor/Private/InputCustomizations.h
--- a/Engine/Plugins/EnhancedInput/Source/InputEditor/Private/InputCustomizations.h
+++ b/Engine/Plugins/EnhancedInput/Source/InputEditor/Private/InputCustomizations.h
@@ -10,15 +10,6 @@ class UPlayerMappableKeySettings;
 class IDetailPropertyRow;
 class IDetailLayoutBuilder;
 
-class FInputContextDetails : public IDetailCustomization
-{
-public:
-	/** Makes a new instance of this detail layout class for a specific detail view requesting it */
-	static TSharedRef<IDetailCustomization> MakeInstance();
-
-	/** ILayoutDetails interface */
-	virtual void CustomizeDetails( IDetailLayoutBuilder& DetailBuilder ) override;
-};
 
 class FEnhancedActionMappingCustomization : public IPropertyTypeCustomization
 {
@@ -132,4 +123,27 @@ private:
 
 	TSharedPtr<IPropertyHandle> SettingsPropHandle;
 	TSharedPtr<IPropertyHandle> MappingNamePropHandle;
+};
+
+/**
+ * Customization for FInputMappingContextMappingData.
+ */
+class FMappingContextKeyContainerCustomization : public IPropertyTypeCustomization
+{
+public:
+	static TSharedRef<IPropertyTypeCustomization> MakeInstance();
+	
+	virtual void CustomizeHeader(
+		TSharedRef<IPropertyHandle> PropertyHandle,
+		FDetailWidgetRow& HeaderRow,
+		IPropertyTypeCustomizationUtils& CustomizationUtils) override;
+		
+	virtual void CustomizeChildren(
+		TSharedRef<IPropertyHandle> PropertyHandle,
+		IDetailChildrenBuilder& ChildBuilder,
+		IPropertyTypeCustomizationUtils& CustomizationUtils) override;
+
+private:
+	
+	TSharedPtr<IPropertyHandle> MappingsArrayHandle;
 };
\ No newline at end of file
diff --git a/Engine/Plugins/EnhancedInput/Source/InputEditor/Private/InputEditorModule.cpp b/Engine/Plugins/EnhancedInput/Source/InputEditor/Private/InputEditorModule.cpp
--- a/Engine/Plugins/EnhancedInput/Source/InputEditor/Private/InputEditorModule.cpp
+++ b/Engine/Plugins/EnhancedInput/Source/InputEditor/Private/InputEditorModule.cpp
@@ -348,7 +348,7 @@ void FInputEditorModule::StartupModule()
 	
 	// Register customizations
 	FPropertyEditorModule& PropertyModule = FModuleManager::GetModuleChecked<FPropertyEditorModule>("PropertyEditor");
-	PropertyModule.RegisterCustomClassLayout("InputMappingContext", FOnGetDetailCustomizationInstance::CreateStatic(&FInputContextDetails::MakeInstance));
+	PropertyModule.RegisterCustomPropertyTypeLayout(FInputMappingContextMappingData::StaticStruct()->GetFName(), FOnGetPropertyTypeCustomizationInstance::CreateStatic(&FMappingContextKeyContainerCustomization::MakeInstance));
 	PropertyModule.RegisterCustomPropertyTypeLayout("EnhancedActionKeyMapping", FOnGetPropertyTypeCustomizationInstance::CreateStatic(&FEnhancedActionMappingCustomization::MakeInstance));
 	PropertyModule.RegisterCustomClassLayout(UEnhancedInputDeveloperSettings::StaticClass()->GetFName(), FOnGetDetailCustomizationInstance::CreateStatic(&FEnhancedInputDeveloperSettingsCustomization::MakeInstance));
 	PropertyModule.RegisterCustomPropertyTypeLayout(UPlayerMappableKeySettings::StaticClass()->GetFName(), FOnGetPropertyTypeCustomizationInstance::CreateStatic(&FPlayerMappableKeyChildSettingsCustomization::MakeInstance));
@@ -392,7 +392,7 @@ void FInputEditorModule::ShutdownModule()
 
 	// Unregister customizations
 	FPropertyEditorModule& PropertyModule = FModuleManager::GetModuleChecked<FPropertyEditorModule>("PropertyEditor");
-	PropertyModule.UnregisterCustomClassLayout("InputContext");
+	PropertyModule.UnregisterCustomPropertyTypeLayout("InputMappingContextMappingData");
 	PropertyModule.UnregisterCustomPropertyTypeLayout("EnhancedActionKeyMapping");
 	PropertyModule.UnregisterCustomClassLayout("EnhancedInputDeveloperSettings");
 	PropertyModule.UnregisterCustomPropertyTypeLayout("PlayerMappableKeySettings");
diff --git a/Engine/Source/Runtime/Core/Public/UObject/FortniteSeasonBranchObjectVersions.inl b/Engine/Source/Runtime/Core/Public/UObject/FortniteSeasonBranchObjectVersions.inl
--- a/Engine/Source/Runtime/Core/Public/UObject/FortniteSeasonBranchObjectVersions.inl
+++ b/Engine/Source/Runtime/Core/Public/UObject/FortniteSeasonBranchObjectVersions.inl
@@ -48,5 +48,8 @@ FORTNITE_SEASON_VERSION(ChaosImplicitObjectUnionLeafObjectsToInt32, 12)
 // Chaos Visual Debugger : Adding serialization for properties that were being recorded, but not serialized
 FORTNITE_SEASON_VERSION(CVDSerializationFixMissingSerializationProperties, 13)
 
+// Updated Enhanceed Input Mapping Contexts to support adding "Profile override" mappings.
+FORTNITE_SEASON_VERSION(EnhancedInputMappingContextProfileMappingsUpdate, 14)
+
 // -----<new versions can be added above this line>-------------------------------------------------
 #undef FORTNITE_SEASON_VERSION
\ No newline at end of file
