[Enhanced Input] Remove the UPlayerMappableInputConfig type, which was previsouly deprecated in 5.3.

#rnx

diff --git a/Engine/Plugins/EnhancedInput/Source/EnhancedInput/Public/EnhancedInputSubsystemInterface.h b/Engine/Plugins/EnhancedInput/Source/EnhancedInput/Public/EnhancedInputSubsystemInterface.h
--- a/Engine/Plugins/EnhancedInput/Source/EnhancedInput/Public/EnhancedInputSubsystemInterface.h
+++ b/Engine/Plugins/EnhancedInput/Source/EnhancedInput/Public/EnhancedInputSubsystemInterface.h
@@ -20,7 +20,6 @@ class UInputAction;
 class UEnhancedPlayerInput;
 class UInputModifier;
 class UInputTrigger;
-class UPlayerMappableInputConfig;
 class UWorldSubsystem;
 class UEnhancedInputUserSettings;
 class UEnhancedPlayerMappableKeyProfile;
diff --git a/Engine/Plugins/EnhancedInput/Source/InputEditor/Private/InputEditorModule.cpp b/Engine/Plugins/EnhancedInput/Source/InputEditor/Private/InputEditorModule.cpp
--- a/Engine/Plugins/EnhancedInput/Source/InputEditor/Private/InputEditorModule.cpp
+++ b/Engine/Plugins/EnhancedInput/Source/InputEditor/Private/InputEditorModule.cpp
@@ -9,7 +9,6 @@
 #include "Framework/Notifications/NotificationManager.h"
 #include "InputMappingContext.h"
 #include "Misc/PackageName.h"
-#include "PlayerMappableInputConfig.h"
 #include "InputCustomizations.h"
 #include "ISettingsModule.h"
 #include "ToolMenuSection.h"
@@ -215,23 +214,6 @@ UObject* UInputAction_Factory::FactoryCreateNew(UClass* Class, UObject* InParent
 	}
 }
 
-PRAGMA_DISABLE_DEPRECATION_WARNINGS
-// UPlayerMappableInputConfig_Factory
-UDEPRECATED_UPlayerMappableInputConfig_Factory::UDEPRECATED_UPlayerMappableInputConfig_Factory(const class FObjectInitializer& OBJ)
-	: Super(OBJ)
-{
-	SupportedClass = UPlayerMappableInputConfig::StaticClass();
-	bEditAfterNew = true;
-	bCreateNew = true;
-}
-
-UObject* UDEPRECATED_UPlayerMappableInputConfig_Factory::FactoryCreateNew(UClass* Class, UObject* InParent, FName Name, EObjectFlags Flags, UObject* Context, FFeedbackContext* Warn)
-{
-	check(Class->IsChildOf(UPlayerMappableInputConfig::StaticClass()));
-	return NewObject<UPlayerMappableInputConfig>(InParent, Class, Name, Flags | RF_Transactional, Context);
-}
-PRAGMA_ENABLE_DEPRECATION_WARNINGS
-
 // Asset type actions
 
 class FAssetTypeActions_InputContext : public FAssetTypeActions_DataAsset
@@ -336,9 +318,6 @@ public:
 		
 		Set("ClassIcon.InputMappingContext", new IMAGE_BRUSH_SVG("Icons/InputMappingContext_16", Icon16));
 		Set("ClassThumbnail.InputMappingContext", new IMAGE_BRUSH_SVG("Icons/InputMappingContext_64", Icon64));
-		
-		Set("ClassIcon.PlayerMappableInputConfig", new IMAGE_BRUSH_SVG("Icons/PlayerMappableInputConfig_16", Icon16));
-		Set("ClassThumbnail.PlayerMappableInputConfig", new IMAGE_BRUSH_SVG("Icons/PlayerMappableInputConfig_64", Icon64));			
 	}
 };
 
diff --git a/Engine/Plugins/EnhancedInput/Source/InputEditor/Public/InputEditorModule.h b/Engine/Plugins/EnhancedInput/Source/InputEditor/Public/InputEditorModule.h
--- a/Engine/Plugins/EnhancedInput/Source/InputEditor/Public/InputEditorModule.h
+++ b/Engine/Plugins/EnhancedInput/Source/InputEditor/Public/InputEditorModule.h
@@ -90,12 +90,4 @@ public:
 	UE_API virtual UObject* FactoryCreateNew(UClass* Class, UObject* InParent, FName Name, EObjectFlags Flags, UObject* Context, FFeedbackContext* Warn) override;
 };
 
-UCLASS(deprecated)
-class UE_DEPRECATED(5.3, "UPlayerMappableInputConfig has been deprecated, please use the UEnhancedInputUserSettings system instead.") INPUTEDITOR_API UDEPRECATED_UPlayerMappableInputConfig_Factory : public UFactory
-{
-	GENERATED_UCLASS_BODY()
-public:
-	virtual UObject* FactoryCreateNew(UClass* Class, UObject* InParent, FName Name, EObjectFlags Flags, UObject* Context, FFeedbackContext* Warn) override;
-};
-
 #undef UE_API
