Hook up the HUD DebugTarget to the Debug Enhanced Input player focus

diff --git a/Engine/Plugins/EnhancedInput/Source/EnhancedInput/Private/EnhancedInputModule.cpp b/Engine/Plugins/EnhancedInput/Source/EnhancedInput/Private/EnhancedInputModule.cpp
--- a/Engine/Plugins/EnhancedInput/Source/EnhancedInput/Private/EnhancedInputModule.cpp
+++ b/Engine/Plugins/EnhancedInput/Source/EnhancedInput/Private/EnhancedInputModule.cpp
@@ -333,14 +333,29 @@ void FEnhancedInputModule::OnShowDebugInfo(AHUD* HUD, UCanvas* Canvas, const FDe
 			DisplayDebugManager.SetFont(GEngine->GetSmallFont());
 			DisplayDebugManager.SetDrawColor(FColor::Yellow);
 			DisplayDebugManager.DrawString(TEXT("ENHANCED INPUT"));
-			
-			// TODO: Support paging through subsystems one at a time (via console? key press?)
 
-			// Show first player only for now
-			TObjectIterator<UEnhancedInputLocalPlayerSubsystem> FirstPlayer;
-			if (FirstPlayer)
+			bool bFoundValidDebugTarget = false;
+			// Checking if the debug target actor is a valid local player, if so then show their debug info
+			// If target actor is not a local player then don't display anything
+			if (APawn* TargetPawn = Cast<APawn>(HUD->GetCurrentDebugTargetActor()))
 			{
-				FirstPlayer->ShowDebugInfo(Canvas);
+				if (APlayerController* PlayerController = Cast<APlayerController>(TargetPawn->GetController()))
+				{
+					if (const ULocalPlayer* LocalPlayer = Cast<ULocalPlayer>(PlayerController->Player))
+					{
+						if (UEnhancedInputLocalPlayerSubsystem* EnhancedInputSubsystem = ULocalPlayer::GetSubsystem<UEnhancedInputLocalPlayerSubsystem>(LocalPlayer))
+						{
+							EnhancedInputSubsystem->ShowDebugInfo(Canvas);
+							bFoundValidDebugTarget = true;
+						}
+					}
+				}
+			}
+
+			if(!bFoundValidDebugTarget)
+			{ 
+				DisplayDebugManager.SetDrawColor(FColor::Red);
+				DisplayDebugManager.DrawString(TEXT("Debug Target is not valid for Enhanced Input Debug Display"));
 			}
 		}
 
