[Enhanced Input] Add a PostUserSettingsInitialized delegate to the Enhanced Input Local Player subsytem. This provides a simple place to get notified of when the User Settings object has been created in your gameplay code, as you may want to respond to some loaded data from it.

[RN] Add a OnPostUserSettingsInitialized delegate to Enhanced Input, which will fire after the Enhanced Input User Settings have been initialized on the subsystem.

diff --git a/Engine/Plugins/EnhancedInput/Source/EnhancedInput/Private/EnhancedInputSubsystems.cpp b/Engine/Plugins/EnhancedInput/Source/EnhancedInput/Private/EnhancedInputSubsystems.cpp
--- a/Engine/Plugins/EnhancedInput/Source/EnhancedInput/Private/EnhancedInputSubsystems.cpp
+++ b/Engine/Plugins/EnhancedInput/Source/EnhancedInput/Private/EnhancedInputSubsystems.cpp
@@ -94,7 +94,7 @@ void UEnhancedInputLocalPlayerSubsystem::InitalizeUserSettings()
 
 	if (!LP)
 	{
-		UE_LOG(LogEnhancedInput, Error, TEXT("Failed to initalize user settings, there is no valid local player!"));
+		UE_LOG(LogEnhancedInput, Error, TEXT("Failed to initialize user settings, there is no valid local player!"));
 		return;
 	}
 	
@@ -125,7 +125,10 @@ void UEnhancedInputLocalPlayerSubsystem::InitalizeUserSettings()
 	if (ensure(UserSettings))
 	{
 		BindUserSettingDelegates();
-		UE_LOG(LogEnhancedInput, Log, TEXT("Enhanced Input local player subsystem has initalized the user settings!"));
+		
+		OnPostUserSettingsInitialized.Broadcast(UserSettings);
+		
+		UE_LOG(LogEnhancedInput, Log, TEXT("Enhanced Input local player subsystem has initialized the user settings!"));
 	}
 }
 
diff --git a/Engine/Plugins/EnhancedInput/Source/EnhancedInput/Public/EnhancedInputSubsystems.h b/Engine/Plugins/EnhancedInput/Source/EnhancedInput/Public/EnhancedInputSubsystems.h
--- a/Engine/Plugins/EnhancedInput/Source/EnhancedInput/Public/EnhancedInputSubsystems.h
+++ b/Engine/Plugins/EnhancedInput/Source/EnhancedInput/Public/EnhancedInputSubsystems.h
@@ -52,6 +52,7 @@ public:
 	DECLARE_DYNAMIC_MULTICAST_DELEGATE(FOnControlMappingsRebuilt);
 	DECLARE_DYNAMIC_MULTICAST_DELEGATE_OneParam(FOnMappingContextAdded, const UInputMappingContext*, MappingContext);
 	DECLARE_DYNAMIC_MULTICAST_DELEGATE_OneParam(FOnMappingContextRemoved, const UInputMappingContext*, MappingContext);
+	DECLARE_DYNAMIC_MULTICAST_DELEGATE_OneParam(FOnUserSettingsPostInitialized, const UEnhancedInputUserSettings*, Settings);
 
 	/**
 	 * Blueprint Event that is called at the end of any frame that Control Mappings have been rebuilt.
@@ -62,7 +63,7 @@ public:
 	/*
 	* A callback fired when a mapping context is added (AddMappingContext is called on this subsystem)
 	* 
-	* Note:  This does not make any gaurantee that the control mappings will have been rebuilt. 
+	* Note:  This does not make any guarantee that the control mappings will have been rebuilt. 
 	* If you need that, then listen to the ControlMappingsRebuiltDelegate instead.
 	*/
 	UPROPERTY(BlueprintAssignable, Category = "Input")
@@ -71,12 +72,18 @@ public:
 	/*
 	* A callback fired when a mapping context is removed (RemoveMappingContext is called on this subsystem)
 	* 
-	* Note:  This does not make any gaurantee that the control mappings will have been rebuilt. 
+	* Note:  This does not make any guarantee that the control mappings will have been rebuilt. 
 	* If you need that, then listen to the ControlMappingsRebuiltDelegate instead.
 	*/
 	UPROPERTY(BlueprintAssignable, Category = "Input")
 	FOnMappingContextRemoved OnMappingContextRemoved;
 
+	/**
+	 * Delegate that is fired after this local player's User Settings object has been created and loaded for the first time.
+	 */
+	UPROPERTY(BlueprintAssignable, Category = "Input")
+	FOnUserSettingsPostInitialized OnPostUserSettingsInitialized;
+
 protected:
     	
 	/** The user settings for this subsystem used to store each user's input related settings */
