[Enhanced Input] Remove an edit condition which did nothing.

#rnx

diff --git a/Engine/Plugins/EnhancedInput/Source/EnhancedInput/Public/UserSettings/EnhancedInputUserSettings.h b/Engine/Plugins/EnhancedInput/Source/EnhancedInput/Public/UserSettings/EnhancedInputUserSettings.h
--- a/Engine/Plugins/EnhancedInput/Source/EnhancedInput/Public/UserSettings/EnhancedInputUserSettings.h
+++ b/Engine/Plugins/EnhancedInput/Source/EnhancedInput/Public/UserSettings/EnhancedInputUserSettings.h
@@ -79,7 +79,7 @@ struct FMapPlayerKeyArgs final
 	FKey NewKey;
 	
 	/** An OPTIONAL specifier about what kind of hardware this mapping is for. */
-	UPROPERTY(EditAnywhere, BlueprintReadWrite, Category="Enhanced Input|User Settings", meta = (EditCondition = "bFilterByHardwareDeviceId == true", AllowPrivateAccess = true, GetOptions = "Engine.InputPlatformSettings.GetAllHardwareDeviceNames"))
+	UPROPERTY(EditAnywhere, BlueprintReadWrite, Category="Enhanced Input|User Settings", meta = (AllowPrivateAccess = true, GetOptions = "Engine.InputPlatformSettings.GetAllHardwareDeviceNames"))
 	FName HardwareDeviceId;
 
 	/** The Key Mapping Profile identifier that this mapping should be set on. If this is empty, then the currently equipped profile will be used. */
