[Enhanced Input]

- Add some nice helpers to Enhanced Input user settings to make it easier to grab the default key profile
- Helper function to get if a key profile is currently available or not

#rnx

diff --git a/Engine/Plugins/EnhancedInput/Source/EnhancedInput/Private/UserSettings/EnhancedInputUserSettings.cpp b/Engine/Plugins/EnhancedInput/Source/EnhancedInput/Private/UserSettings/EnhancedInputUserSettings.cpp
--- a/Engine/Plugins/EnhancedInput/Source/EnhancedInput/Private/UserSettings/EnhancedInputUserSettings.cpp
+++ b/Engine/Plugins/EnhancedInput/Source/EnhancedInput/Private/UserSettings/EnhancedInputUserSettings.cpp
@@ -1492,6 +1492,20 @@ const TMap<FString, TObjectPtr<UEnhancedPlayerMappableKeyProfile>>& UEnhancedInp
 	return SavedKeyProfilesMap;
 }
 
+bool UEnhancedInputUserSettings::IsKeyProfileAvailable(const FString& ProfileId) const
+{
+	return SavedKeyProfilesMap.Contains(ProfileId);
+}
+
+UEnhancedPlayerMappableKeyProfile* UEnhancedInputUserSettings::GetDefaultKeyProfile() const
+{
+	if (const TObjectPtr<UEnhancedPlayerMappableKeyProfile>* Found = SavedKeyProfilesMap.Find(UE::EnhancedInput::DefaultProfileIdentifierString))
+	{
+		return Found->Get();
+	}
+	return nullptr;
+}
+
 UEnhancedPlayerMappableKeyProfile* UEnhancedInputUserSettings::CreateNewKeyProfile(const FPlayerMappableKeyProfileCreationArgs& InArgs)
 {
 	if (!InArgs.ProfileType)
diff --git a/Engine/Plugins/EnhancedInput/Source/EnhancedInput/Public/UserSettings/EnhancedInputUserSettings.h b/Engine/Plugins/EnhancedInput/Source/EnhancedInput/Public/UserSettings/EnhancedInputUserSettings.h
--- a/Engine/Plugins/EnhancedInput/Source/EnhancedInput/Public/UserSettings/EnhancedInputUserSettings.h
+++ b/Engine/Plugins/EnhancedInput/Source/EnhancedInput/Public/UserSettings/EnhancedInputUserSettings.h
@@ -733,6 +733,15 @@ public:
 	/** Returns all player saved key profiles */
 	UE_API const TMap<FString, TObjectPtr<UEnhancedPlayerMappableKeyProfile>>& GetAllAvailableKeyProfiles() const;
 
+	/** Returns true if the given ProfileId is currently available in the user settings. */
+	UE_API bool IsKeyProfileAvailable(const FString& ProfileId) const;
+
+	/**
+	* Returns a pointer to the default key profile created by Enhanced Input. 
+	* This profile has an ID of "InputUserSettings.Profiles.Default"
+	*/
+	UE_API UEnhancedPlayerMappableKeyProfile* GetDefaultKeyProfile() const;
+
 	/**
 	 * Creates a new profile with this name and type.
 	 */
