[Enhanced Input Settings] Deprecate the virtual OnKeyMappingRegistered function in favor of one which takes a pointer to the key profile that the mapping was registered to, OnKeyMappingRegisteredToProfile.

#pf 6893d13fe31477d5b66e4fdc
#rnx

diff --git a/Engine/Plugins/EnhancedInput/Source/EnhancedInput/Private/UserSettings/EnhancedInputUserSettings.cpp b/Engine/Plugins/EnhancedInput/Source/EnhancedInput/Private/UserSettings/EnhancedInputUserSettings.cpp
--- a/Engine/Plugins/EnhancedInput/Source/EnhancedInput/Private/UserSettings/EnhancedInputUserSettings.cpp
+++ b/Engine/Plugins/EnhancedInput/Source/EnhancedInput/Private/UserSettings/EnhancedInputUserSettings.cpp
@@ -1033,6 +1033,14 @@ void UEnhancedInputUserSettings::Serialize(FArchive& Ar)
 	// Now we can actually re-populate our runtime data with the info from the headers that we just loaded in
 	if (Ar.IsLoading())
 	{
+		// When loading, we want to empty the current saved key profile map so that it is populated with 
+		// only data from the archive.
+		//
+		// We only care about the dirty data when serializing. If other profiles currently exist here, they were added 
+		// via the Registration of IMC assets and would need to be re-registered. We will re-populate the registered IMCs
+		// after we load everything that was dirty.
+		SavedKeyProfilesMap.Empty();		
+
 		for (const UE::EnhancedInput::FMappableKeysHeader& Header : Headers)
 		{
 			if (const UClass* FoundClass = FindObject<UClass>(nullptr, *Header.ClassPath))
@@ -1057,11 +1065,9 @@ void UEnhancedInputUserSettings::Serialize(FArchive& Ar)
 					MappingRow.Mappings.Add(PlayerMapping);
 				}
 
-				// We need to populate this key profile with all the known key mappings so that it's up to date 
-				for (const TObjectPtr<const UInputMappingContext>& IMC : RegisteredMappingContexts)
-				{
-					RegisterKeyMappingsToProfile(*NewProfile, IMC);
-				}
+
+				// Note: The currently registered Input Mapping contexts will be re-registered on these
+				// newly created profiles after we are done loading all previous data.
 
 				// Ensure that the owning local player is up to date
 				if (const ULocalPlayer* LP = GetLocalPlayer())
@@ -1105,6 +1111,18 @@ void UEnhancedInputUserSettings::Serialize(FArchive& Ar)
 		}
 #endif	// WITH_EDITOR
 	}
+
+	// Now that we have finished loading all the data that was in the archive, we should re-register all the 
+	// mapping contexts to the newly created key profiles.
+	if (Ar.IsLoading())
+	{
+		// This will re-create any profiles which may have existed before this serialize load call was initiated and ensure
+		// that each profile has its default key mappings set up for each IMC which has unique defaults for specific profiles
+		for (const TObjectPtr<const UInputMappingContext>& IMC : RegisteredMappingContexts)
+		{
+			RegisterInputMappingContextInternal(IMC);
+		}
+	}
 }
 
 ULocalPlayer* UEnhancedInputUserSettings::GetLocalPlayer() const
@@ -1595,6 +1613,18 @@ bool UEnhancedInputUserSettings::RegisterInputMappingContext(const UInputMapping
 	// Keep track of all the registered IMC's
 	RegisteredMappingContexts.Add(IMC);
 
+	return RegisterInputMappingContextInternal(IMC);
+}
+
+bool UEnhancedInputUserSettings::RegisterInputMappingContextInternal(const UInputMappingContext* IMC)
+{
+	if (!IMC)
+	{
+		UE_LOG(LogEnhancedInput, Error, TEXT("Attempting to register a null mapping context with the user settings!"));
+		ensure(false);
+		return false;
+	}	
+
 	bool bResult = true;
 
 	const ULocalPlayer* LP = GetLocalPlayer();
@@ -1724,7 +1754,7 @@ bool UEnhancedInputUserSettings::RegisterKeyMappingsToProfile(UEnhancedPlayerMap
 		{
 			if (RegisteredMappingId.IsValidId())
 			{
-				OnKeyMappingRegistered(MappingRow.Mappings.Get(RegisteredMappingId), KeyMapping);
+				OnKeyMappingRegisteredToProfile(&Profile, MappingRow.Mappings.Get(RegisteredMappingId), KeyMapping);
 			}			
 		}
 
@@ -1736,10 +1766,22 @@ bool UEnhancedInputUserSettings::RegisterKeyMappingsToProfile(UEnhancedPlayerMap
 	return true;
 }
 
+PRAGMA_DISABLE_DEPRECATION_WARNINGS
 void UEnhancedInputUserSettings::OnKeyMappingRegistered(FPlayerKeyMapping& RegisteredMapping, const FEnhancedActionKeyMapping& SourceMapping)
 {
 	// Does nothing by default
 }
+PRAGMA_ENABLE_DEPRECATION_WARNINGS
+
+void UEnhancedInputUserSettings::OnKeyMappingRegisteredToProfile(UEnhancedPlayerMappableKeyProfile* Profile, FPlayerKeyMapping& RegisteredMapping, const FEnhancedActionKeyMapping& SourceMapping)
+{
+	// Call the old version for back compat
+	PRAGMA_DISABLE_DEPRECATION_WARNINGS
+	OnKeyMappingRegistered(RegisteredMapping, SourceMapping);
+	PRAGMA_ENABLE_DEPRECATION_WARNINGS
+
+	// Otherwise, do nothing by default
+}
 
 FHardwareDeviceIdentifier UEnhancedInputUserSettings::DetermineHardwareDeviceForActionMapping(const FEnhancedActionKeyMapping& ActionMapping, const UInputMappingContext* IMC) const
 {
diff --git a/Engine/Plugins/EnhancedInput/Source/EnhancedInput/Public/UserSettings/EnhancedInputUserSettings.h b/Engine/Plugins/EnhancedInput/Source/EnhancedInput/Public/UserSettings/EnhancedInputUserSettings.h
--- a/Engine/Plugins/EnhancedInput/Source/EnhancedInput/Public/UserSettings/EnhancedInputUserSettings.h
+++ b/Engine/Plugins/EnhancedInput/Source/EnhancedInput/Public/UserSettings/EnhancedInputUserSettings.h
@@ -601,7 +601,25 @@ protected:
 	 * @param RegisteredMapping		The newly registered player key mapping
 	 * @param SourceMapping			The source mapping that the registered mapping was created from. 
 	 */
-	UE_API virtual void OnKeyMappingRegistered(FPlayerKeyMapping& RegisteredMapping, const FEnhancedActionKeyMapping& SourceMapping);
+	UE_DEPRECATED(5.7, "Use OnKeyMappingRegisteredToProfile instead.")
+	UE_API virtual void OnKeyMappingRegistered(FPlayerKeyMapping& RegisteredMapping, const FEnhancedActionKeyMapping& SourceMapping) final;
+
+
+	/**
+	 * A virtual function that will be called whenever a player key mapping is first registered. This will be called once
+	 * by UEnhancedInputUserSettings::RegisterKeyMappingsToProfile for each mapping that is created from the Input Mapping Context.
+	 *
+	 * This is a good place to do any "Post Registration" setup you may want to for you key mappings such as changing the
+	 * key that is mapped based on some other settings or data outside of this system (i.e. you are upgrading to this system
+	 * and already had a key mapping system in place, you can set the key here!)
+	 *
+	 * This does nothing by default.
+	 * 
+	 * @param Profile				The key profile which this mapping was registered to
+	 * @param RegisteredMapping		The newly registered player key mapping
+	 * @param SourceMapping			The source mapping that the registered mapping was created from.
+	 */
+	UE_API virtual void OnKeyMappingRegisteredToProfile(UEnhancedPlayerMappableKeyProfile* Profile, FPlayerKeyMapping& RegisteredMapping, const FEnhancedActionKeyMapping& SourceMapping);
 
 public:
 	
@@ -810,6 +828,16 @@ public:
 
 protected:
 
+	/**
+	* This will register the given IMC to the settings system. 
+	* 
+	* If an IMC specified mappings which should have different defaults based on a certain key profile ID,
+	* then we will also create the profile here if needed.
+	* 
+	* This will do so regardless of if the IMC has been registered already. 
+	*/
+	bool RegisterInputMappingContextInternal(const UInputMappingContext* IMC);
+
 	/**
 	 * This will register individual key mappings from the given Input Mapping context on the profile id.
 	 *
