[Enhanced Input] Expose GetMappingsForProfile with UE_API for games to use.

#Rnx

diff --git a/Engine/Plugins/EnhancedInput/Source/EnhancedInput/Public/InputMappingContext.h b/Engine/Plugins/EnhancedInput/Source/EnhancedInput/Public/InputMappingContext.h
--- a/Engine/Plugins/EnhancedInput/Source/EnhancedInput/Public/InputMappingContext.h
+++ b/Engine/Plugins/EnhancedInput/Source/EnhancedInput/Public/InputMappingContext.h
@@ -199,7 +199,7 @@ public:
 	 * @return Reference to the key mappings to use for the given profile. If there is no override mapping specified for the given
 	 *		   profile, then the default mappings will be returned.
 	 */
-	const TArray<FEnhancedActionKeyMapping>& GetMappingsForProfile(const FString& ProfileId) const;
+	UE_API const TArray<FEnhancedActionKeyMapping>& GetMappingsForProfile(const FString& ProfileId) const;
 
 	/**
 	 * Returns true if this Input Mapping Context contains a reference to the given Input Action in its default or profile override mappings.
