[Enhanced Input] Add the ability for bound listeners on a UEnhancedInputComponent to consume lower priority listeners.

Also, add a ENHANCED_INPUT_ALLOW_LEGACY_BINDING define check before deleting the legacy input binding functions. Having these available for your game may make it easier to ugprade and older project to use Enhanced Input over time.

[RN] You can now call "SetShouldConsume" on any Enhanced Input bound action handles. You can also define "ENHANCED_INPUT_ALLOW_LEGACY_BINDING" in your game's .build.cs file to enable legacy key bindings onto enhanced input components, making upgrading to enhanced input easier.

diff --git a/Engine/Plugins/EnhancedInput/Source/EnhancedInput/Private/EnhancedPlayerInput.cpp b/Engine/Plugins/EnhancedInput/Source/EnhancedInput/Private/EnhancedPlayerInput.cpp
--- a/Engine/Plugins/EnhancedInput/Source/EnhancedInput/Private/EnhancedPlayerInput.cpp
+++ b/Engine/Plugins/EnhancedInput/Source/EnhancedInput/Private/EnhancedPlayerInput.cpp
@@ -26,6 +26,11 @@ namespace UE::Input
 	static FAutoConsoleVariableRef CVarReconcileRemovedMappingDelegates(TEXT("EnhancedInput.ReconcileRemovedMappingDelegates"),
 		ReconcileRemovedMappingDelegates,
 		TEXT("When the mappings are rebuilt, they have been 'in process'. If this is true, we will set their value to zero so that they fire the 'Canceled' event on their next evaluation."));
+
+	static bool bEnableListenerConsumption = true;
+	static FAutoConsoleVariableRef CVarEnableListenerConsumption(TEXT("EnhancedInput.EnableListenerConsumption"),
+		bEnableListenerConsumption,
+		TEXT("If true, then bound listeners have the option to consume any Enhanced Input delegates on lower priority input components."));
 }
 
 namespace UE::EnhancedInput
@@ -554,6 +559,16 @@ void UEnhancedPlayerInput::EvaluateInputDelegates(const TArray<UInputComponent*>
 
 	// Execute appropriate delegates
 
+	// A set of input actions whose bound delete was flagged to consume the input, preventing
+	// other listeners from receiving the event. This is how Enhanced Input respects the priority of UInputComponent's
+	// This allows other UEnhancedInputComponents (notably, not any "legacy" UInputComponents) completely consume other
+	// listeners bound to the same input action.
+	//
+	// This is useful if you have some gameplay code which you want to "override" the behavior of other listeners in some specific
+	// scenario, like if you have a specific higher priority gameplay ability or something like that. This way, you don't need to
+	// manually be tracking some additional game state constantly. 
+	TSet<const UInputAction*> ConsumedInputActions;
+
 #if DEV_ONLY_KEY_BINDINGS_AVAILABLE
 	// Cache modifier key states for debug key bindings
 	const bool bAlt = IsAltPressed(), bCtrl = IsCtrlPressed(), bShift = IsShiftPressed(), bCmd = IsCmdPressed();
@@ -626,6 +641,13 @@ void UEnhancedPlayerInput::EvaluateInputDelegates(const TArray<UInputComponent*>
 			TObjectPtr<const UInputAction> DelegateAction = Delegate->GetAction();
 			bool bCanTrigger = true;
 
+			// Skip input actions which have been marked for consumption by a higher priority UEnhancedInputComponent
+			if (UE::Input::bEnableListenerConsumption && ConsumedInputActions.Contains(DelegateAction))
+			{
+				UE_LOG(LogEnhancedInput, Verbose, TEXT("Input action '%s' was consumed by a higher priority input component listener. It will no longer fire"), *DelegateAction->GetName());
+				continue;
+			}
+
 			if (UE::Input::ShouldOnlyTriggerLastActionInChord)
 			{
 				// If this delegate is referenced by a UInputTriggerChordAction::ChordAction
@@ -668,6 +690,13 @@ void UEnhancedPlayerInput::EvaluateInputDelegates(const TArray<UInputComponent*>
 						Delegate->Execute(*ActionData);
 					}					
 				}
+
+				// Keep track of any delegates which want to consume input away from other lower priority listeners
+				// bound to this same input action.
+				if (UE::Input::bEnableListenerConsumption && Delegate->ShouldConsume())
+				{
+					ConsumedInputActions.Add(DelegateAction);
+				}
 			}
 		}
 		TriggeredDelegates.Reset();
diff --git a/Engine/Plugins/EnhancedInput/Source/EnhancedInput/Public/EnhancedInputComponent.h b/Engine/Plugins/EnhancedInput/Source/EnhancedInput/Public/EnhancedInputComponent.h
--- a/Engine/Plugins/EnhancedInput/Source/EnhancedInput/Public/EnhancedInputComponent.h
+++ b/Engine/Plugins/EnhancedInput/Source/EnhancedInput/Public/EnhancedInputComponent.h
@@ -176,17 +176,34 @@ private:
 	/** Trigger event that raises the delegate */
 	ETriggerEvent TriggerEvent = ETriggerEvent::None;
 
+	/**
+	 * If true, then this event binding will stop any other listeners to the same UInputAction with a lower input component priority
+	 * on the stack.
+	 */
+	uint8 bConsumes : 1 = false;
+
 protected:
 	// Clone constructor
-	FEnhancedInputActionEventBinding(const FEnhancedInputActionEventBinding& CloneFrom, EInputBindingClone Clone) : FInputBindingHandle(CloneFrom, Clone), Action(CloneFrom.Action), TriggerEvent(CloneFrom.TriggerEvent) {}
+	FEnhancedInputActionEventBinding(const FEnhancedInputActionEventBinding& CloneFrom, EInputBindingClone Clone) :
+		FInputBindingHandle(CloneFrom, Clone),
+		Action(CloneFrom.Action), 
+		TriggerEvent(CloneFrom.TriggerEvent), 
+		bConsumes(CloneFrom.bConsumes) {}
 
 public:
 	FEnhancedInputActionEventBinding() = default;
-	FEnhancedInputActionEventBinding(const UInputAction* InAction, ETriggerEvent InTriggerEvent) : Action(InAction), TriggerEvent(InTriggerEvent) {}
+	FEnhancedInputActionEventBinding(
+		const UInputAction* InAction,
+		ETriggerEvent InTriggerEvent,
+		const bool bRequestsConsumption = false)
+	: Action(InAction), TriggerEvent(InTriggerEvent), bConsumes(bRequestsConsumption) {}
 
 	const UInputAction* GetAction() const { return Action.Get(); }
+	bool ShouldConsume() const { return bConsumes; } 
 	ETriggerEvent GetTriggerEvent() const { return TriggerEvent; }
 
+	void SetShouldConsume(const bool bNewConsumeValue) { bConsumes = bNewConsumeValue; }
+	
 	virtual void Execute(const FInputActionInstance& ActionData) const = 0;
 	virtual TUniquePtr<FEnhancedInputActionEventBinding> Clone() const = 0;
 	virtual void SetShouldFireWithEditorScriptGuard(const bool bNewValue) = 0;
@@ -265,8 +282,8 @@ private:
 		Delegate.SetShouldFireWithEditorScriptGuard(CloneFrom.Delegate.ShouldFireWithEditorScriptGuard());
 	}
 public:
-	FEnhancedInputActionEventDelegateBinding(const UInputAction* Action, ETriggerEvent InTriggerEvent)
-		: FEnhancedInputActionEventBinding(Action, InTriggerEvent)
+	FEnhancedInputActionEventDelegateBinding(const UInputAction* Action, ETriggerEvent InTriggerEvent, const bool bRequestsConsumption = false)
+		: FEnhancedInputActionEventBinding(Action, InTriggerEvent, bRequestsConsumption)
 	{}
 
 	// Implemented below.
@@ -563,6 +580,19 @@ public:
 	UFUNCTION(BlueprintCallable, meta = (BlueprintInternalUseOnly = "true", HideSelfPin = "true", HidePin = "Action"))
 	UE_API FInputActionValue GetBoundActionValue(const UInputAction* Action) const;
 
+
+	// While we do not encourage the binding of legacy action and axis input bindings on a UEnhancedInputComponent,
+	// there is no functional issues with doing so. Set this define in your game's .build.cs file if you would like
+	// allow for legacy binding to Enhanced Input Components. This maybe useful if you are upgrading an older project to
+	// use Enhanced Input.
+	//
+	// You can set this in the .build.cs file like this:
+	//		PublicDefinitions.Add("ENHANCED_INPUT_ALLOW_LEGACY_BINDING=1");
+#ifndef ENHANCED_INPUT_ALLOW_LEGACY_BINDING
+	#define ENHANCED_INPUT_ALLOW_LEGACY_BINDING 0
+#endif
+
+#if !ENHANCED_INPUT_ALLOW_LEGACY_BINDING
 	// Delete all InputComponent binding helpers. Indicates intentions going forward and improves intellisense/VAX when working with EnhancedInputCompoennts.
 	template<class UserClass>
 	FInputActionBinding& BindAction(const FName ActionName, const EInputEvent KeyEvent, UserClass* Object, typename FInputActionHandlerSignature::TMethodPtr< UserClass > Func) = delete;
@@ -587,6 +617,7 @@ public:
 	FInputTouchBinding& BindTouch(const EInputEvent KeyEvent, UserClass* Object, typename FInputTouchHandlerSignature::TMethodPtr< UserClass > Func) = delete;
 	template<class UserClass>
 	FInputGestureBinding& BindGesture(const FKey GestureKey, UserClass* Object, typename FInputGestureHandlerSignature::TMethodPtr< UserClass > Func) = delete;
+#endif	// !ENHANCED_INPUT_ALLOW_LEGACY_BINDING
 };
 
 #undef UE_API
diff --git a/Engine/Plugins/EnhancedInput/Source/EnhancedInput/Public/InputMappingContext.h b/Engine/Plugins/EnhancedInput/Source/EnhancedInput/Public/InputMappingContext.h
--- a/Engine/Plugins/EnhancedInput/Source/EnhancedInput/Public/InputMappingContext.h
+++ b/Engine/Plugins/EnhancedInput/Source/EnhancedInput/Public/InputMappingContext.h
@@ -187,12 +187,12 @@ public:
 	 * @param ProfileId The Player Mappable Key Profile ID
 	 * @return True if this Input Mapping Context has some specific override mappings for the given profile. 
 	 */
-	bool HasMappingsForProfile(const FString& ProfileId) const;
+	UE_API bool HasMappingsForProfile(const FString& ProfileId) const;
 
 	/**
 	 * @return An array of profile identifiers which this IMC specifies overrides for.
 	 */
-	TArray<FString> GetProfilesWithOverridenMappings() const;
+	UE_API TArray<FString> GetProfilesWithOverridenMappings() const;
 
 	/**
 	 * @param ProfileId The Player Mappable Key Profile ID
