PR #13668: Added more context in the EnhancedActionKeyMapping DataValidation when there is a null trigger on a key mapping.

[RN] The data validation error for when there is a null input trigger on a key mapping now includes the name of the input action which it is mapped to.

diff --git a/Engine/Plugins/EnhancedInput/Source/EnhancedInput/Private/EnhancedActionKeyMapping.cpp b/Engine/Plugins/EnhancedInput/Source/EnhancedInput/Private/EnhancedActionKeyMapping.cpp
--- a/Engine/Plugins/EnhancedInput/Source/EnhancedInput/Private/EnhancedActionKeyMapping.cpp
+++ b/Engine/Plugins/EnhancedInput/Source/EnhancedInput/Private/EnhancedActionKeyMapping.cpp
@@ -131,7 +131,9 @@ EDataValidationResult FEnhancedActionKeyMapping::IsDataValid(FDataValidationCont
 		else
 		{
 			Result = EDataValidationResult::Invalid;
-			Context.AddError(LOCTEXT("NullInputTrigger", "There cannot be a null Input Trigger on a key mapping"));
+			FFormatNamedArguments Args;
+			Args.Add("InputActionName", FText::FromString(GetNameSafe(Action)));
+			Context.AddError(FText::Format(LOCTEXT("NullInputTriggerWithAction", "There cannot be a null Input Trigger on a key mapping for Input Action '{InputActionName}'"), Args));
 		}
 	}
 
@@ -200,7 +202,9 @@ EDataValidationResult FEnhancedActionKeyMapping::IsDataValid(FDataValidationCont
 		else
 		{
 			Result = EDataValidationResult::Invalid;
-			Context.AddError(LOCTEXT("NullInputModifier", "There cannot be a null Input Modifier on a key mapping"));
+			FFormatNamedArguments Args;
+			Args.Add("InputActionName", FText::FromString(GetNameSafe(Action)));
+			Context.AddError(FText::Format(LOCTEXT("NullInputModifierWithAction", "There cannot be a null Input Modifier on a key mapping for Input Action '{InputActionName}'"), Args));
 		}
 	}
 	return Result;
