PR #13144: [Enhanced Input] Fix for UE-270589.

Input triggers and modifiers will now properly respect the priority of an input mapping context. This fixes a bug where if you had the same input modifier class on an input mapping context of a different priorty, it would always use the modifier values of the one which was pushed to the stack first.

You can disable this behavior with the "input.bRespectIMCPriortyForTriggers" CVar, in case you were unknowningly relying on this behavior.

[RN] Input triggers and modifiers will now properly respect the priority of an input mapping context. This fixes a bug where if you had the same input modifier class on an input mapping context of a different priorty, it would always use the modifier values of the one which was pushed to the stack first.

diff --git a/Engine/Plugins/EnhancedInput/Source/EnhancedInput/Private/EnhancedInputSubsystemInterface.cpp b/Engine/Plugins/EnhancedInput/Source/EnhancedInput/Private/EnhancedInputSubsystemInterface.cpp
--- a/Engine/Plugins/EnhancedInput/Source/EnhancedInput/Private/EnhancedInputSubsystemInterface.cpp
+++ b/Engine/Plugins/EnhancedInput/Source/EnhancedInput/Private/EnhancedInputSubsystemInterface.cpp
@@ -33,6 +33,14 @@ static FAutoConsoleVariableRef GCVarGlobalAxisConfigMode(
 	TEXT("Whether or not to apply Global Axis Config settings. 0 = Default (Mouse Only), 1 = All, 2 = None")
 );
 
+static bool bRespectIMCPriorityForTriggers = true;
+static FAutoConsoleVariableRef CVarRespectIMCPriorityForTriggers(
+	TEXT("input.bRespectIMCPriortyForTriggers"),
+	bRespectIMCPriorityForTriggers,
+	TEXT("If true, then when an IMC with the same triggers but a higher priority is added, the triggers on the higher priority IMC will be used.\n")
+		 TEXT("This CVar will be removed in a future release.")
+);
+
 template<typename T>
 void DeepCopyPtrArray(const TArray<T*>& From, TArray<T*>& To)
 {
@@ -1110,6 +1118,14 @@ void IEnhancedInputSubsystemInterface::RebuildControlMappings()
 				});
 			if (Idx != INDEX_NONE)
 			{
+				// Fix for UE-270589, behind a CVar for now just in case projects may have been relying on the old behavior
+				// and need a switch off to retain functionality. 
+				if (bRespectIMCPriorityForTriggers)
+				{
+					OldMappings[Idx].Triggers = MoveTemp(Mapping.Triggers);
+					OldMappings[Idx].Modifiers = MoveTemp(Mapping.Modifiers);
+				}
+
 				Mapping = MoveTemp(OldMappings[Idx]);
 				OldMappings.RemoveAtSwap(Idx);
 			}
