[Enhanced Input] Add a commented out "OptimizeCode = CodeOptimization.Never;" to the Enhanced Input build CS file to give users an example of how to more easily debug that engine module.

You can uncomment this line in the build.cs to turn off optimization for the whole Enhanced Input module, which is really useful for seeing FNames or other more complex values.

#rnx

diff --git a/Engine/Plugins/EnhancedInput/Source/EnhancedInput/EnhancedInput.Build.cs b/Engine/Plugins/EnhancedInput/Source/EnhancedInput/EnhancedInput.Build.cs
--- a/Engine/Plugins/EnhancedInput/Source/EnhancedInput/EnhancedInput.Build.cs
+++ b/Engine/Plugins/EnhancedInput/Source/EnhancedInput/EnhancedInput.Build.cs
@@ -6,6 +6,9 @@ namespace UnrealBuildTool.Rules
 	{
         public EnhancedInput(ReadOnlyTargetRules Target) : base(Target)
         {
+			// Uncomment this line to make for easier debugging
+	        //OptimizeCode = CodeOptimization.Never;
+	        
             PrivateDependencyModuleNames.AddRange(
 				new string[] {
                     "ApplicationCore",
