Fix a possible null de-reference in UEnhancedPlayerInput::GetActionValue

#rnx

diff --git a/Engine/Plugins/EnhancedInput/Source/EnhancedInput/Private/EnhancedPlayerInput.cpp b/Engine/Plugins/EnhancedInput/Source/EnhancedInput/Private/EnhancedPlayerInput.cpp
--- a/Engine/Plugins/EnhancedInput/Source/EnhancedInput/Private/EnhancedPlayerInput.cpp
+++ b/Engine/Plugins/EnhancedInput/Source/EnhancedInput/Private/EnhancedPlayerInput.cpp
@@ -836,6 +836,12 @@ void UEnhancedPlayerInput::FlushPressedKeys()
 
 FInputActionValue UEnhancedPlayerInput::GetActionValue(TObjectPtr<const UInputAction> ForAction) const
 {
+	if (!ForAction)
+	{
+		UE_LOG(LogEnhancedInput, Error, TEXT("[%hs] Attempting to get the value of an invalid Input action!"), __func__);
+		return FInputActionValue();
+	}
+
 	const FInputActionInstance* ActionData = FindActionInstanceData(ForAction);
 	return ActionData ? ActionData->GetValue() : FInputActionValue(ForAction->ValueType, FInputActionValue::Axis3D::ZeroVector);
 }
