[Enhanced Input] Do not load platform specific override data on post load of the platform settings. If you are referencing an object that is in a plugin, this will just fail anyway and does not improve any sort of loading behavior.

Ideally, in the future, we should be using a more modular design to allow users to add/remove platform specific override data right on the subsystem, perhaps via a Game Feature Action.

#rn Enhanced Input: Do not load platform specific override data on post load of the platform settings. Wait until runtime to load them so that any assets in plugins will work and not cause warnings on boot.

diff --git a/Engine/Plugins/EnhancedInput/Source/EnhancedInput/Private/EnhancedInputPlatformSettings.cpp b/Engine/Plugins/EnhancedInput/Source/EnhancedInput/Private/EnhancedInputPlatformSettings.cpp
--- a/Engine/Plugins/EnhancedInput/Source/EnhancedInput/Private/EnhancedInputPlatformSettings.cpp
+++ b/Engine/Plugins/EnhancedInput/Source/EnhancedInput/Private/EnhancedInputPlatformSettings.cpp
@@ -58,25 +58,6 @@ const UInputMappingContext* UEnhancedInputPlatformData::GetContextRedirect(UInpu
 //////////////////////////////////////////////////////////////////////////
 // UEnhancedInputPlatformSettings
 
-void UEnhancedInputPlatformSettings::PostLoad()
-{
-	Super::PostLoad();
-
-	InputDataClasses.Reset();
-	LoadInputDataClasses();
-}
-
-#if WITH_EDITOR
-void UEnhancedInputPlatformSettings::PostEditChangeProperty(FPropertyChangedEvent& PropertyChangedEvent)
-{
-	Super::PostEditChangeProperty(PropertyChangedEvent);
-
-	// Reset the cached data classes so that they get re-evaluated the next time we want to load them
-	// This ensures that we always have valid cached values during PIE sessions 
-	InputDataClasses.Reset();
-}
-#endif	// WITH_EDITOR
-
 void UEnhancedInputPlatformSettings::GetAllMappingContextRedirects(OUT TMap<TObjectPtr<const UInputMappingContext>, TObjectPtr<const UInputMappingContext>>& OutRedirects)
 {
 	ForEachInputData([&OutRedirects](const UEnhancedInputPlatformData& Data)
@@ -87,17 +68,19 @@ void UEnhancedInputPlatformSettings::GetAllMappingContextRedirects(OUT TMap<TObj
 
 void UEnhancedInputPlatformSettings::ForEachInputData(TFunctionRef<void(const UEnhancedInputPlatformData&)> Predicate)
 {
-	LoadInputDataClasses();
-	
-	for (TSubclassOf<UEnhancedInputPlatformData> InputDataClass : InputDataClasses)
+	for (TSoftClassPtr<UEnhancedInputPlatformData> InputDataPtr : InputData)
 	{
-		if (const UEnhancedInputPlatformData* Data = InputDataClass.GetDefaultObject())
+		if (TSubclassOf<UEnhancedInputPlatformData> InputDataClass = InputDataPtr.LoadSynchronous())
 		{
-			Predicate(*Data);
+			if (const UEnhancedInputPlatformData* Data = InputDataClass.GetDefaultObject())
+			{
+				Predicate(*Data);
+			}
 		}
 	}
 }
 
+PRAGMA_DISABLE_DEPRECATION_WARNINGS
 void UEnhancedInputPlatformSettings::LoadInputDataClasses()
 {
 	if (InputData.Num() != InputDataClasses.Num())
@@ -111,5 +94,6 @@ void UEnhancedInputPlatformSettings::LoadInputDataClasses()
 		}
 	}
 }
+PRAGMA_ENABLE_DEPRECATION_WARNINGS
 
 #undef LOCTEXT_NAMESPACE
diff --git a/Engine/Plugins/EnhancedInput/Source/EnhancedInput/Public/EnhancedInputPlatformSettings.h b/Engine/Plugins/EnhancedInput/Source/EnhancedInput/Public/EnhancedInputPlatformSettings.h
--- a/Engine/Plugins/EnhancedInput/Source/EnhancedInput/Public/EnhancedInputPlatformSettings.h
+++ b/Engine/Plugins/EnhancedInput/Source/EnhancedInput/Public/EnhancedInputPlatformSettings.h
@@ -9,6 +9,7 @@
 #define UE_API ENHANCEDINPUT_API
 
 class UInputMappingContext;
+struct FStreamableHandle;
 
 /**
  * A base class that can be used to store platform specific data for Enhanced Input.
@@ -58,13 +59,6 @@ public:
 		return UPlatformSettingsManager::Get().GetSettingsForPlatform<UEnhancedInputPlatformSettings>();
 	}
 
-	//~ UObject interface
-	UE_API virtual void PostLoad() override;
-#if WITH_EDITOR
-	UE_API virtual void PostEditChangeProperty(struct FPropertyChangedEvent& PropertyChangedEvent) override;
-#endif
-	//~ End UObject interface
-
 	/** Get the current Enhanced Input platform data */
 	const TArray<TSoftClassPtr<UEnhancedInputPlatformData>>& GetInputData() const { return InputData; }
 
@@ -80,6 +74,7 @@ public:
 protected:
 
 	/** Load the input data subclasses and cache them for later use in the InputDataClasses array. */
+	UE_DEPRECATED(5.7, "LoadInputDataClasses is now deprecated")
 	UE_API void LoadInputDataClasses();
 	
 	/** Input data that can be populated with Enhanced Input Platform Data blueprints */
@@ -90,6 +85,7 @@ protected:
 	 * A transient array of the subclasses for the Enhanced Input Platform data. This will prevent us from
 	 * from having to load the class' default object during game time.
 	 */
+	UE_DEPRECATED(5.7, "InputDataClasses is now deprecated")
 	UPROPERTY(Transient)
 	TArray<TSubclassOf<UEnhancedInputPlatformData>> InputDataClasses;
 
