Added a some cvars to improve enhanced input debug scalability by allowing the displayed data to be reduced.
- EnhancedInput.Debug.ShowKeyboardKeys
- EnhancedInput.Debug.ShowGamepadKeys
- EnhancedInput.Debug.ShowTouchKeys
- EnhancedInput.Debug.InputActionFilter
- EnhancedInput.Debug.ContextMappingFilter

diff --git a/Engine/Plugins/EnhancedInput/Source/EnhancedInput/Private/EnhancedInputSubsystemInterfaceDebug.cpp b/Engine/Plugins/EnhancedInput/Source/EnhancedInput/Private/EnhancedInputSubsystemInterfaceDebug.cpp
--- a/Engine/Plugins/EnhancedInput/Source/EnhancedInput/Private/EnhancedInputSubsystemInterfaceDebug.cpp
+++ b/Engine/Plugins/EnhancedInput/Source/EnhancedInput/Private/EnhancedInputSubsystemInterfaceDebug.cpp
@@ -16,6 +16,34 @@
  * See EnhancedInputSubsystemInterface.cpp for main functionality.
  */
 
+namespace UE::EnhancedInput
+{
+	static int32 EIDebug_ShowKeyboardKeys = 1;
+	static FAutoConsoleVariableRef CVarEIDebug_ShowKeyboardKeys(TEXT("EnhancedInput.Debug.ShowKeyboardKeys"),
+		EIDebug_ShowKeyboardKeys,
+		TEXT("Should Enhanced Input Debug show keyboard keys in the debug info?"));
+	
+	static int32 EIDebug_ShowGamepadKeys = 1;
+	static FAutoConsoleVariableRef CVarEIDebug_ShowGamepadKeys(TEXT("EnhancedInput.Debug.ShowGamepadKeys"),
+		EIDebug_ShowGamepadKeys,
+		TEXT("Should Enhanced Input Debug show Gamepad keys in the debug info?"));
+		
+	static int32 EIDebug_ShowTouchKeys = 1;
+	static FAutoConsoleVariableRef CVarEIDebug_ShowTouchKeys(TEXT("EnhancedInput.Debug.ShowTouchKeys"),
+		EIDebug_ShowTouchKeys,
+		TEXT("Should Enhanced Input Debug show Touch keys in the debug info?"));
+		
+	static FString EIDebug_InputActionFilter = TEXT("");
+	static FAutoConsoleVariableRef CVarEIDebug_InputActionFilter(TEXT("EnhancedInput.Debug.InputActionFilter"),
+		EIDebug_InputActionFilter,
+		TEXT("This is the list of filtered input actions to display. Multiple filter entries should be separated by a comma."));
+		
+	static FString EIDebug_ContextMappingFilter = TEXT("");
+	static FAutoConsoleVariableRef CVarEIDebug_ContextMappingFilter(TEXT("EnhancedInput.Debug.ContextMappingFilter"),
+		EIDebug_ContextMappingFilter,
+		TEXT("This is the list of filtered mapping contexts to display. Multiple filter entries should be separated by a comma."));
+}
+
 // Manual GC protection is required for visualization textures as they live within an interface class.
 struct FVisualizationTexture
 {
@@ -120,6 +148,35 @@ void IEnhancedInputSubsystemInterface::ShowMappingContextDebugInfo(UCanvas* Canv
 				continue;
 			}
 
+			//Update the array of filtered entries
+			static FString LastAppliedContextMappingFilter = TEXT("");
+			static TArray<FString> AppliedContextMappingFilters;
+			if (!LastAppliedContextMappingFilter.Equals(UE::EnhancedInput::EIDebug_ContextMappingFilter, ESearchCase::IgnoreCase))
+			{
+				AppliedContextMappingFilters.Reset();
+				LastAppliedContextMappingFilter = UE::EnhancedInput::EIDebug_ContextMappingFilter;
+				LastAppliedContextMappingFilter.ParseIntoArray(AppliedContextMappingFilters, TEXT(","));
+				for (FString& Entry : AppliedContextMappingFilters)
+				{
+					Entry.TrimStartAndEndInline();
+				}
+			}
+
+			//apply the filter
+			bool bMatchedContextMappingFilter = AppliedContextMappingFilters.IsEmpty();
+			for (const FString& FilterEntry : AppliedContextMappingFilters)
+			{
+				if (AppliedContext->GetName().Contains(FilterEntry, ESearchCase::IgnoreCase))
+				{
+					bMatchedContextMappingFilter = true;
+					break;
+				}
+			}
+			if (!bMatchedContextMappingFilter)
+			{
+				continue;
+			}			
+
 			DisplayDebugManager.SetDrawColor(FColor::Yellow);
 
 			const EMappingContextRegistrationTrackingMode TrackingMode = AppliedContext->GetRegistrationTrackingMode();
@@ -197,6 +254,35 @@ void IEnhancedInputSubsystemInterface::ShowMappingContextDebugInfo(UCanvas* Canv
 
 			for (const UInputAction* Action : OrderedActions)
 			{
+				//Update the array of filtered entries
+				static FString LastAppliedInputActionFilter = TEXT("");
+				static TArray<FString> AppliedInputActionFilters;
+				if (!LastAppliedInputActionFilter.Equals(UE::EnhancedInput::EIDebug_InputActionFilter, ESearchCase::IgnoreCase))
+				{
+					AppliedInputActionFilters.Reset();
+					LastAppliedInputActionFilter = UE::EnhancedInput::EIDebug_InputActionFilter;
+					LastAppliedInputActionFilter.ParseIntoArray(AppliedInputActionFilters, TEXT(","));
+					for (FString& Entry : AppliedInputActionFilters)
+					{
+						Entry.TrimStartAndEndInline();
+					}
+				}
+
+				//apply the filter
+				bool bMatchedInputActionFilter = AppliedInputActionFilters.IsEmpty();
+				for (const FString& FilterEntry : AppliedInputActionFilters)
+				{
+					if (Action->GetName().Contains(FilterEntry, ESearchCase::IgnoreCase))
+					{
+						bMatchedInputActionFilter = true;
+						break;
+					}
+				}
+				if (!bMatchedInputActionFilter)
+				{
+					continue;
+				}
+				
 				struct FDisplayLine
 				{
 					FDisplayLine(FColor InColor, FString InLine) : Color(InColor), Line(InLine) {}
@@ -226,6 +312,24 @@ void IEnhancedInputSubsystemInterface::ShowMappingContextDebugInfo(UCanvas* Canv
 				// TODO: Order these by key display name?
 				for (FEnhancedActionKeyMapping& Mapping : ActionMappings[Action])
 				{
+					const bool bIsGamepadKey = Mapping.Key.IsGamepadKey();
+					const bool bIsTouchKey = Mapping.Key.IsTouch();
+					const bool bIsKeyboardKey = !bIsGamepadKey && !bIsTouchKey;
+
+					//skip based on key type and cvars
+					if (bIsGamepadKey && !UE::EnhancedInput::EIDebug_ShowGamepadKeys)
+					{
+						continue;
+					}
+					else if (bIsTouchKey && !UE::EnhancedInput::EIDebug_ShowTouchKeys)
+					{
+						continue;
+					}
+					else if (bIsKeyboardKey && !UE::EnhancedInput::EIDebug_ShowKeyboardKeys)
+					{
+						continue;
+					}
+					
 					FString* KeyOwner = AppliedKeys.Find(Mapping.Key.GetFName());
 					FColor DrawColor = FColor::White;
 					const FColor DrawColorInvalid = FColor(64, 64, 64);
