[Enhanced Input] Add a null check in RebuildControlMappings to combat a reported crash in a users project.

#rnx

diff --git a/Engine/Plugins/EnhancedInput/Source/EnhancedInput/Private/EnhancedInputSubsystemInterface.cpp b/Engine/Plugins/EnhancedInput/Source/EnhancedInput/Private/EnhancedInputSubsystemInterface.cpp
--- a/Engine/Plugins/EnhancedInput/Source/EnhancedInput/Private/EnhancedInputSubsystemInterface.cpp
+++ b/Engine/Plugins/EnhancedInput/Source/EnhancedInput/Private/EnhancedInputSubsystemInterface.cpp
@@ -952,6 +952,12 @@ void IEnhancedInputSubsystemInterface::RebuildControlMappings()
 		// Don't apply context specific keys immediately, allowing multiple mappings to the same key within the same context if required.
 		TArray<FKey> ContextAppliedKeys;
 		const UInputMappingContext* MappingContext = ContextPair.Key;
+
+		if (!ensure(IsValid(MappingContext)))
+		{
+			UE_LOG(LogEnhancedInput, Error, TEXT("[%hs] Encountered an invalid mapping context while rebuilding control mappings. Some inputs may be ignored!"), __func__);
+			continue;
+		}
 		
 		// Check if this mapping context can be applied in our current input mode on the player
 		// If it can't, then we will not process its mappings..
