[Enhanced Input] Add a console command "EnhancedInput.RebuildControlMappings" which will call RequestRebuildControlMappings on every available input subsystem. This can be useful for debugging things which may happen only on rebuild of control mappings, or if you want to test your code which may respone to a mapping context change

#rnx

diff --git a/Engine/Plugins/EnhancedInput/Source/EnhancedInput/Private/EnhancedInputLibrary.cpp b/Engine/Plugins/EnhancedInput/Source/EnhancedInput/Private/EnhancedInputLibrary.cpp
--- a/Engine/Plugins/EnhancedInput/Source/EnhancedInput/Private/EnhancedInputLibrary.cpp
+++ b/Engine/Plugins/EnhancedInput/Source/EnhancedInput/Private/EnhancedInputLibrary.cpp
@@ -10,12 +10,30 @@
 #include "EnhancedInputSubsystems.h"
 #include "GameFramework/Actor.h"
 #include "GameFramework/PlayerController.h"
+#include "HAL/IConsoleManager.h"
 #include "InputTriggers.h"
 #include "PlayerMappableKeySettings.h"
 #include "UObject/UObjectIterator.h"
 
 #include UE_INLINE_GENERATED_CPP_BY_NAME(EnhancedInputLibrary)
 
+namespace UE::Input
+{
+	static void RebuildAllControlMappings(const TArray<FString>& Args)
+	{
+		UEnhancedInputLibrary::ForEachSubsystem([Args](IEnhancedInputSubsystemInterface* Subsystem)
+		{
+			Subsystem->RequestRebuildControlMappings();
+		});
+	}
+
+	static FAutoConsoleCommand ConsoleCommandDumpProfileToLog(
+		TEXT("EnhancedInput.RebuildControlMappings"),
+		TEXT("Calls RequestRebuildControlMappings on every available Enhanced Input Subsystem"),
+		FConsoleCommandWithArgsDelegate::CreateStatic(UE::Input::RebuildAllControlMappings));
+}
+
+
 void UEnhancedInputLibrary::ForEachSubsystem(TFunctionRef<void(IEnhancedInputSubsystemInterface*)> SubsystemPredicate)
 {
 	// World subsystem
