Fixed issues with overlapping keys in different modes

Can be hotfixed off:
- EnhancedInput.TrackKeysForAllActionMappings - Set to FALSE to return to only tracking keys for mappings that are set to consume legacy keys.

#tests PIE with common problem configs

diff --git a/Engine/Plugins/EnhancedInput/Source/EnhancedInput/Private/EnhancedInputSubsystemInterface.cpp b/Engine/Plugins/EnhancedInput/Source/EnhancedInput/Private/EnhancedInputSubsystemInterface.cpp
--- a/Engine/Plugins/EnhancedInput/Source/EnhancedInput/Private/EnhancedInputSubsystemInterface.cpp
+++ b/Engine/Plugins/EnhancedInput/Source/EnhancedInput/Private/EnhancedInputSubsystemInterface.cpp
@@ -41,6 +41,13 @@ static FAutoConsoleVariableRef CVarRespectIMCPriorityForTriggers(
 		 TEXT("This CVar will be removed in a future release.")
 );
 
+static bool bTrackKeysForAllActionMappings = false;
+static FAutoConsoleVariableRef GCVarTrackKeysForAllActionMappings(
+	TEXT("EnhancedInput.TrackKeysForAllActionMappings"),
+	bTrackKeysForAllActionMappings,
+	TEXT("If true, keep track of ALL keys mapped to each action, regardless of whether that mapping has been configured to consume legacy keys. This is needed for Enhanced Input bindings to respect legacy key consumption.")
+);
+
 template<typename T>
 void DeepCopyPtrArray(const TArray<T*>& From, TArray<T*>& To)
 {
@@ -78,6 +85,11 @@ UEnhancedInputUserSettings* IEnhancedInputSubsystemInterface::GetUserSettings()
 	return nullptr;
 }
 
+bool IEnhancedInputSubsystemInterface::TrackingKeysForAllActionMappings()
+{
+	return bTrackKeysForAllActionMappings;
+}
+
 void IEnhancedInputSubsystemInterface::BindUserSettingDelegates()
 {
 	UEnhancedInputUserSettings* Settings = GetUserSettings();
@@ -1005,10 +1017,20 @@ void IEnhancedInputSubsystemInterface::RebuildControlMappings()
 			{
 				Mapping.Key = PlayerMappedKey;
 
-				// If this Input Action is flagged to consume input, then mark it's key state as being consumed every tick.
-				// This has the affect where the base UPlayerInput class will not fire any legacy bindings
-				if (Mapping.Action->bConsumesActionAndAxisMappings)
+				if (bTrackKeysForAllActionMappings)
+				{
+					// Track all key mappings for each action so that we can check if they've already been consumed before executing their action.
+					FKeyConsumptionOptions& Opts = PlayerInput->KeyConsumptionData.FindOrAdd(Mapping.Action);
+					Opts.KeysToConsume.AddUnique(Mapping.Key);
+					if (Mapping.Action->bConsumesActionAndAxisMappings)
+					{
+						Opts.EventsToCauseConsumption |= static_cast<ETriggerEvent>(Mapping.Action->TriggerEventsThatConsumeLegacyKeys);
+					}
+				}
+				else if (Mapping.Action->bConsumesActionAndAxisMappings)
 				{
+					// If this Input Action is flagged to consume input, then mark it's key state as being consumed every tick.
+					// This has the effect where the base UPlayerInput class will not fire any legacy bindings
 					FKeyConsumptionOptions& Opts = PlayerInput->KeyConsumptionData.FindOrAdd(Mapping.Action);
 					Opts.KeysToConsume.AddUnique(Mapping.Key);
 					Opts.EventsToCauseConsumption |= static_cast<ETriggerEvent>(Mapping.Action->TriggerEventsThatConsumeLegacyKeys);
diff --git a/Engine/Plugins/EnhancedInput/Source/EnhancedInput/Private/EnhancedPlayerInput.cpp b/Engine/Plugins/EnhancedInput/Source/EnhancedInput/Private/EnhancedPlayerInput.cpp
--- a/Engine/Plugins/EnhancedInput/Source/EnhancedInput/Private/EnhancedPlayerInput.cpp
+++ b/Engine/Plugins/EnhancedInput/Source/EnhancedInput/Private/EnhancedPlayerInput.cpp
@@ -7,6 +7,7 @@
 #include "EnhancedInputComponent.h"
 #include "EnhancedInputDeveloperSettings.h"
 #include "EnhancedInputModule.h"
+#include "EnhancedInputSubsystemInterface.h"
 #include "GameFramework/PlayerController.h"
 #include "GameFramework/WorldSettings.h"
 #include "HAL/IConsoleManager.h"
@@ -610,6 +611,9 @@ void UEnhancedPlayerInput::PrepareInputDelegatesForEvaluation(const TArray<UInpu
 bool UEnhancedPlayerInput::EvaluateInputComponentDelegates(UInputComponent* const InputComponent, const TArray<TPair<FKey, FKeyState*>>& KeysWithEvents, const float DeltaTime, const bool bGamePaused)
 {
 	bool bBlocksInput = false;
+	const bool bCheckConsumedKeys = IEnhancedInputSubsystemInterface::TrackingKeysForAllActionMappings();
+
+	static TArray<FKey> KeysToConsume;
 
 #if DEV_ONLY_KEY_BINDINGS_AVAILABLE
 	// Cache modifier key states for debug key bindings
@@ -626,11 +630,59 @@ bool UEnhancedPlayerInput::EvaluateInputComponentDelegates(UInputComponent* cons
 			// PERF: Lots of map lookups! Group EnhancedActionBindings by Action?
 			if (const FInputActionInstance* ActionData = FindActionInstanceData(Binding->GetAction()))
 			{
+				const FKeyConsumptionOptions* const ConsumptionData = KeyConsumptionData.Find(ActionData->GetSourceAction());
+
 				const ETriggerEvent BoundTriggerEvent = Binding->GetTriggerEvent();
 				// Raise appropriate delegate to report on event state
 				if (ActionData->TriggerEvent == BoundTriggerEvent ||
 					(BoundTriggerEvent == ETriggerEvent::Started && ActionData->TriggerEventInternal == ETriggerEventInternal::StartedAndTriggered))	// Triggering in a single tick should also fire the started event.
 				{
+					if (bCheckConsumedKeys)
+					{
+						// Don't do this action if its triggering keys were consumed elsewhere
+						if (ConsumptionData)
+						{
+							bool bHasActiveConsumedKey = false;
+							bool bHasActiveUnconsumedKey = false;
+
+							for (const FKey& KeyToConsume : ConsumptionData->KeysToConsume)
+							{
+								// Skip any action whose triggering keys have already been consumed.
+								if (FKeyState* KeyState = KeyStateMap.Find(KeyToConsume))
+								{
+									if (KeyState->bDown || KeyState->bDownPrevious)
+									{
+										if (KeyState->bConsumed)
+										{
+											bHasActiveConsumedKey = true;
+										}
+										else
+										{
+											bHasActiveUnconsumedKey = true;
+											KeysToConsume.AddUnique(KeyToConsume);
+										}
+									}
+								}
+							}
+
+							// We can't just skip an Action if we don't find any unconsumed Keys associated with it, because some actions may not have a key (eg: if they were Injected).
+							// Therefore, only skip this action if both:
+							//  1. It is associated with one or more Key that is or was pressed but has already been consumed -- this indicates that a conflict with a higher-in-the-stack component HAS occurred;
+							//  2. It is NOT associated with any Key that is or was pressed but has NOT already consumed -- as, despite any other conflicts that may have occurred, this would clearly show a non-conflicting request for this action has also occurred.
+							if (bHasActiveConsumedKey && !bHasActiveUnconsumedKey)
+							{
+								if (UE::Input::bEnableListenerConsumption)
+								{
+									ConsumedInputActions.Add(Binding->GetAction());
+								}
+								else
+								{
+									continue;
+								}
+							}
+						}
+					}
+
 					// Record intent to trigger started as well as triggered
 					// EmplaceAt 0 for the "Started" event it is always guaranteed to fire before Triggered
 					if (BoundTriggerEvent == ETriggerEvent::Started)
@@ -650,16 +702,26 @@ bool UEnhancedPlayerInput::EvaluateInputComponentDelegates(UInputComponent* cons
 				}
 
 				// If this delegate is bound to an action that has an event and is flagged to consume legacy keys, then mark it as such.
-				if (const FKeyConsumptionOptions* const ConsumptionData = KeyConsumptionData.Find(ActionData->GetSourceAction()))
+				if (ConsumptionData)
 				{
 					if (static_cast<uint8>(ConsumptionData->EventsToCauseConsumption & ActionData->TriggerEvent) != 0)
 					{
 						// Consume all keys that are mapped to this input action with the proper trigger values
-						for (const FKey& KeyToConsume : ConsumptionData->KeysToConsume)
+						if (bCheckConsumedKeys)
 						{
-							if (FKeyState* KeyState = KeyStateMap.Find(KeyToConsume))
+							for (const FKey& KeyToConsume : ConsumptionData->KeysToConsume)
 							{
-								KeyState->bConsumed = true;
+								KeysToConsume.AddUnique(KeyToConsume);
+							}
+						}
+						else
+						{
+							for (const FKey& KeyToConsume : ConsumptionData->KeysToConsume)
+							{
+								if (FKeyState* KeyState = KeyStateMap.Find(KeyToConsume))
+								{
+									KeyState->bConsumed = true;
+								}
 							}
 						}
 					}
@@ -798,11 +860,26 @@ bool UEnhancedPlayerInput::EvaluateInputComponentDelegates(UInputComponent* cons
 		}
 		TriggeredDebugDelegates.Reset();
 #endif
+
+		bBlocksInput = IC->bBlockInput;
 	}
 	
-	
 	const bool bLegacyBlocksInput = Super::EvaluateInputComponentDelegates(InputComponent, KeysWithEvents, DeltaTime, bGamePaused);
 
+	if (bCheckConsumedKeys && !bBlocksInput)
+	{
+		// we do this after finishing the whole component, so we don't consume a key while there might be more bindings to it
+		for (int32 KeyIndex = 0; KeyIndex < KeysToConsume.Num(); ++KeyIndex)
+		{
+			if (FKeyState* KeyState = KeyStateMap.Find(KeysToConsume[KeyIndex]))
+			{
+				KeyState->bConsumed = true;
+			}
+		}
+	}
+
+	KeysToConsume.Reset();
+
 	return bBlocksInput || bLegacyBlocksInput;
 }
 
diff --git a/Engine/Plugins/EnhancedInput/Source/EnhancedInput/Public/EnhancedInputSubsystemInterface.h b/Engine/Plugins/EnhancedInput/Source/EnhancedInput/Public/EnhancedInputSubsystemInterface.h
--- a/Engine/Plugins/EnhancedInput/Source/EnhancedInput/Public/EnhancedInputSubsystemInterface.h
+++ b/Engine/Plugins/EnhancedInput/Source/EnhancedInput/Public/EnhancedInputSubsystemInterface.h
@@ -113,6 +113,11 @@ public:
 	UFUNCTION(BlueprintCallable, Category="Enhanced Input|User Settings", meta=(DisplayName="Get Enhanced Input User Settings"))
 	UE_API virtual UEnhancedInputUserSettings* GetUserSettings() const;
 
+	/**
+	 * Getter for the CVar EnhancedInput.TrackKeysForAllActionMappings which enables Enhanced Input to respect Key consumption.
+	 */
+	static bool TrackingKeysForAllActionMappings();
+
 protected:
 	
 	/**
