Improved consumption between Enhanced Input and legacy input

Some actions were slipping through. Be more serious about blocking things.

#tests PIE

diff --git a/Engine/Plugins/EnhancedInput/Source/EnhancedInput/Private/EnhancedPlayerInput.cpp b/Engine/Plugins/EnhancedInput/Source/EnhancedInput/Private/EnhancedPlayerInput.cpp
--- a/Engine/Plugins/EnhancedInput/Source/EnhancedInput/Private/EnhancedPlayerInput.cpp
+++ b/Engine/Plugins/EnhancedInput/Source/EnhancedInput/Private/EnhancedPlayerInput.cpp
@@ -631,58 +631,61 @@ bool UEnhancedPlayerInput::EvaluateInputComponentDelegates(UInputComponent* cons
 			if (const FInputActionInstance* ActionData = FindActionInstanceData(Binding->GetAction()))
 			{
 				const FKeyConsumptionOptions* const ConsumptionData = KeyConsumptionData.Find(ActionData->GetSourceAction());
+				bool bAlreadyConsumed = false;
 
-				const ETriggerEvent BoundTriggerEvent = Binding->GetTriggerEvent();
-				// Raise appropriate delegate to report on event state
-				if (ActionData->TriggerEvent == BoundTriggerEvent ||
-					(BoundTriggerEvent == ETriggerEvent::Started && ActionData->TriggerEventInternal == ETriggerEventInternal::StartedAndTriggered))	// Triggering in a single tick should also fire the started event.
+				if (bCheckConsumedKeys && ActionData->LastTriggerState != ETriggerState::None)
 				{
-					if (bCheckConsumedKeys)
+					// Don't do this action if its triggering keys were consumed elsewhere
+					if (ConsumptionData)
 					{
-						// Don't do this action if its triggering keys were consumed elsewhere
-						if (ConsumptionData)
-						{
-							bool bHasActiveConsumedKey = false;
-							bool bHasActiveUnconsumedKey = false;
+						bool bHasActiveConsumedKey = false;
+						bool bHasActiveUnconsumedKey = false;
 
-							for (const FKey& KeyToConsume : ConsumptionData->KeysToConsume)
+						for (const FKey& KeyToConsume : ConsumptionData->KeysToConsume)
+						{
+							KeysToConsume.AddUnique(KeyToConsume);
+							// Skip any action whose triggering keys have already been consumed.
+							if (FKeyState* KeyState = KeyStateMap.Find(KeyToConsume))
 							{
-								// Skip any action whose triggering keys have already been consumed.
-								if (FKeyState* KeyState = KeyStateMap.Find(KeyToConsume))
+								if (KeyState->bDown || KeyState->bDownPrevious)
 								{
-									if (KeyState->bDown || KeyState->bDownPrevious)
+									if (KeyState->bConsumed)
+									{
+										bHasActiveConsumedKey = true;
+									}
+									else
 									{
-										if (KeyState->bConsumed)
-										{
-											bHasActiveConsumedKey = true;
-										}
-										else
-										{
-											bHasActiveUnconsumedKey = true;
-											KeysToConsume.AddUnique(KeyToConsume);
-										}
+										bHasActiveUnconsumedKey = true;
 									}
 								}
 							}
+						}
+
+						bAlreadyConsumed = true;
 
-							// We can't just skip an Action if we don't find any unconsumed Keys associated with it, because some actions may not have a key (eg: if they were Injected).
-							// Therefore, only skip this action if both:
-							//  1. It is associated with one or more Key that is or was pressed but has already been consumed -- this indicates that a conflict with a higher-in-the-stack component HAS occurred;
-							//  2. It is NOT associated with any Key that is or was pressed but has NOT already consumed -- as, despite any other conflicts that may have occurred, this would clearly show a non-conflicting request for this action has also occurred.
-							if (bHasActiveConsumedKey && !bHasActiveUnconsumedKey)
+						// We can't just skip an Action if we don't find any unconsumed Keys associated with it, because some actions may not have a key (eg: if they were Injected).
+						// Therefore, only skip this action if both:
+						//  1. It is associated with one or more Key that is or was pressed but has already been consumed -- this indicates that a conflict with a higher-in-the-stack component HAS occurred;
+						//  2. It is NOT associated with any Key that is or was pressed but has NOT already consumed -- as, despite any other conflicts that may have occurred, this would clearly show a non-conflicting request for this action has also occurred.
+						if (bHasActiveConsumedKey && !bHasActiveUnconsumedKey)
+						{
+							if (UE::Input::bEnableListenerConsumption)
 							{
-								if (UE::Input::bEnableListenerConsumption)
-								{
-									ConsumedInputActions.Add(Binding->GetAction());
-								}
-								else
-								{
-									continue;
-								}
+								ConsumedInputActions.Add(Binding->GetAction());
+							}
+							else
+							{
+								continue;
 							}
 						}
 					}
+				}
 
+				const ETriggerEvent BoundTriggerEvent = Binding->GetTriggerEvent();
+				// Raise appropriate delegate to report on event state
+				if (ActionData->TriggerEvent == BoundTriggerEvent ||
+					(BoundTriggerEvent == ETriggerEvent::Started && ActionData->TriggerEventInternal == ETriggerEventInternal::StartedAndTriggered))	// Triggering in a single tick should also fire the started event.
+				{
 					// Record intent to trigger started as well as triggered
 					// EmplaceAt 0 for the "Started" event it is always guaranteed to fire before Triggered
 					if (BoundTriggerEvent == ETriggerEvent::Started)
@@ -702,26 +705,16 @@ bool UEnhancedPlayerInput::EvaluateInputComponentDelegates(UInputComponent* cons
 				}
 
 				// If this delegate is bound to an action that has an event and is flagged to consume legacy keys, then mark it as such.
-				if (ConsumptionData)
+				if (ConsumptionData && !bAlreadyConsumed)
 				{
 					if (static_cast<uint8>(ConsumptionData->EventsToCauseConsumption & ActionData->TriggerEvent) != 0)
 					{
 						// Consume all keys that are mapped to this input action with the proper trigger values
-						if (bCheckConsumedKeys)
+						for (const FKey& KeyToConsume : ConsumptionData->KeysToConsume)
 						{
-							for (const FKey& KeyToConsume : ConsumptionData->KeysToConsume)
+							if (FKeyState* KeyState = KeyStateMap.Find(KeyToConsume))
 							{
-								KeysToConsume.AddUnique(KeyToConsume);
-							}
-						}
-						else
-						{
-							for (const FKey& KeyToConsume : ConsumptionData->KeysToConsume)
-							{
-								if (FKeyState* KeyState = KeyStateMap.Find(KeyToConsume))
-								{
-									KeyState->bConsumed = true;
-								}
+								KeyState->bConsumed = true;
 							}
 						}
 					}
