Special case for correct input priorities across old and new actions

#tests Playing in Editor across a variety of modes that had previously caused issues

diff --git a/Engine/Plugins/EnhancedInput/Source/EnhancedInput/Private/EnhancedInputComponent.cpp b/Engine/Plugins/EnhancedInput/Source/EnhancedInput/Private/EnhancedInputComponent.cpp
--- a/Engine/Plugins/EnhancedInput/Source/EnhancedInput/Private/EnhancedInputComponent.cpp
+++ b/Engine/Plugins/EnhancedInput/Source/EnhancedInput/Private/EnhancedInputComponent.cpp
@@ -14,6 +14,11 @@ UEnhancedInputComponent::UEnhancedInputComponent(const FObjectInitializer& Objec
 	bBlockInput = false;
 }
 
+bool UEnhancedInputComponent::ShouldCheckForLegacyKeyConsumption() const
+{
+	return bShouldCheckForLegacyKeyConsumption;
+}
+
 void UEnhancedInputComponent::SetShouldFireDelegatesInEditor(const bool bInNewValue)
 {
 	bShouldFireDelegatesInEditor = bInNewValue;
diff --git a/Engine/Plugins/EnhancedInput/Source/EnhancedInput/Private/EnhancedPlayerInput.cpp b/Engine/Plugins/EnhancedInput/Source/EnhancedInput/Private/EnhancedPlayerInput.cpp
--- a/Engine/Plugins/EnhancedInput/Source/EnhancedInput/Private/EnhancedPlayerInput.cpp
+++ b/Engine/Plugins/EnhancedInput/Source/EnhancedInput/Private/EnhancedPlayerInput.cpp
@@ -625,6 +625,7 @@ bool UEnhancedPlayerInput::EvaluateInputComponentDelegates(UInputComponent* cons
 	{
 		// Trigger bound event delegates
 		static TArray<TUniquePtr<FEnhancedInputActionEventBinding>> TriggeredDelegates;
+		const bool bShouldCheckForLegacyKeyConsumption = IC->ShouldCheckForLegacyKeyConsumption();
 		for (const TUniquePtr<FEnhancedInputActionEventBinding>& Binding : IC->GetActionEventBindings())
 		{
 			// PERF: Lots of map lookups! Group EnhancedActionBindings by Action?
@@ -633,7 +634,7 @@ bool UEnhancedPlayerInput::EvaluateInputComponentDelegates(UInputComponent* cons
 				const FKeyConsumptionOptions* const ConsumptionData = KeyConsumptionData.Find(ActionData->GetSourceAction());
 				bool bAlreadyConsumed = false;
 
-				if (bCheckConsumedKeys && ActionData->LastTriggerState != ETriggerState::None)
+				if (bCheckConsumedKeys && bShouldCheckForLegacyKeyConsumption && ActionData->LastTriggerState != ETriggerState::None)
 				{
 					// Don't do this action if its triggering keys were consumed elsewhere
 					if (ConsumptionData)
diff --git a/Engine/Plugins/EnhancedInput/Source/EnhancedInput/Public/EnhancedInputComponent.h b/Engine/Plugins/EnhancedInput/Source/EnhancedInput/Public/EnhancedInputComponent.h
--- a/Engine/Plugins/EnhancedInput/Source/EnhancedInput/Public/EnhancedInputComponent.h
+++ b/Engine/Plugins/EnhancedInput/Source/EnhancedInput/Public/EnhancedInputComponent.h
@@ -394,8 +394,23 @@ protected:
 	 */
 	bool bShouldFireDelegatesInEditor = false;
 
+	/** Whether to check when evaluating the input stack whether legacy inputs can consume keys  */
+	bool bShouldCheckForLegacyKeyConsumption = false;
+
 public:
 
+	/**
+	 * Sets whether this component supports checking for legacy input binds that might consume keys
+	 *
+	 */
+	void SetShouldCheckForLegacyKeyConsumption(bool bNewValue) { bShouldCheckForLegacyKeyConsumption = bNewValue; }
+
+	/**
+	 * Whether this component supports checking for legacy input binds that might consume keys
+	 *
+	 */
+	UE_API bool ShouldCheckForLegacyKeyConsumption() const;
+
 	UE_API void SetShouldFireDelegatesInEditor(const bool bInNewValue);
 	
 	bool ShouldFireDelegatesInEditor() const { return bShouldFireDelegatesInEditor; }
