Fixed priority regression when legacy inputs on the same component as enhanced inputs

#tests PIE with affected actions

diff --git a/Engine/Plugins/EnhancedInput/Source/EnhancedInput/Private/EnhancedPlayerInput.cpp b/Engine/Plugins/EnhancedInput/Source/EnhancedInput/Private/EnhancedPlayerInput.cpp
--- a/Engine/Plugins/EnhancedInput/Source/EnhancedInput/Private/EnhancedPlayerInput.cpp
+++ b/Engine/Plugins/EnhancedInput/Source/EnhancedInput/Private/EnhancedPlayerInput.cpp
@@ -614,6 +614,7 @@ bool UEnhancedPlayerInput::EvaluateInputComponentDelegates(UInputComponent* cons
 	const bool bCheckConsumedKeys = IEnhancedInputSubsystemInterface::TrackingKeysForAllActionMappings();
 
 	static TArray<FKey> KeysToConsume;
+	static TArray<FKey> KeysToConsumeBeforeLegacy;
 
 #if DEV_ONLY_KEY_BINDINGS_AVAILABLE
 	// Cache modifier key states for debug key bindings
@@ -632,9 +633,8 @@ bool UEnhancedPlayerInput::EvaluateInputComponentDelegates(UInputComponent* cons
 			if (const FInputActionInstance* ActionData = FindActionInstanceData(Binding->GetAction()))
 			{
 				const FKeyConsumptionOptions* const ConsumptionData = KeyConsumptionData.Find(ActionData->GetSourceAction());
-				bool bAlreadyConsumed = false;
 
-				if (bCheckConsumedKeys && bShouldCheckForLegacyKeyConsumption && ActionData->LastTriggerState != ETriggerState::None)
+				if (bCheckConsumedKeys && bShouldCheckForLegacyKeyConsumption && (ActionData->LastTriggerState != ETriggerState::None || ActionData->TriggerEvent != ETriggerEvent::None))
 				{
 					// Don't do this action if its triggering keys were consumed elsewhere
 					if (ConsumptionData)
@@ -662,8 +662,6 @@ bool UEnhancedPlayerInput::EvaluateInputComponentDelegates(UInputComponent* cons
 							}
 						}
 
-						bAlreadyConsumed = true;
-
 						// We can't just skip an Action if we don't find any unconsumed Keys associated with it, because some actions may not have a key (eg: if they were Injected).
 						// Therefore, only skip this action if both:
 						//  1. It is associated with one or more Key that is or was pressed but has already been consumed -- this indicates that a conflict with a higher-in-the-stack component HAS occurred;
@@ -706,7 +704,7 @@ bool UEnhancedPlayerInput::EvaluateInputComponentDelegates(UInputComponent* cons
 				}
 
 				// If this delegate is bound to an action that has an event and is flagged to consume legacy keys, then mark it as such.
-				if (ConsumptionData && !bAlreadyConsumed)
+				if (ConsumptionData)
 				{
 					if (static_cast<uint8>(ConsumptionData->EventsToCauseConsumption & ActionData->TriggerEvent) != 0)
 					{
@@ -715,7 +713,7 @@ bool UEnhancedPlayerInput::EvaluateInputComponentDelegates(UInputComponent* cons
 						{
 							if (FKeyState* KeyState = KeyStateMap.Find(KeyToConsume))
 							{
-								KeyState->bConsumed = true;
+								KeysToConsumeBeforeLegacy.AddUnique(KeyToConsume);
 							}
 						}
 					}
@@ -857,12 +855,21 @@ bool UEnhancedPlayerInput::EvaluateInputComponentDelegates(UInputComponent* cons
 
 		bBlocksInput = IC->bBlockInput;
 	}
+
+	// Actions that specifically opt into consuming legacy keys need to do so before we process legacy bindings on this component
+	for (int32 KeyIndex = 0; KeyIndex < KeysToConsumeBeforeLegacy.Num(); ++KeyIndex)
+	{
+		if (FKeyState* KeyState = KeyStateMap.Find(KeysToConsumeBeforeLegacy[KeyIndex]))
+		{
+			KeyState->bConsumed = true;
+		}
+	}
 	
 	const bool bLegacyBlocksInput = Super::EvaluateInputComponentDelegates(InputComponent, KeysWithEvents, DeltaTime, bGamePaused);
 
 	if (bCheckConsumedKeys && !bBlocksInput)
 	{
-		// we do this after finishing the whole component, so we don't consume a key while there might be more bindings to it
+		// Remaining keys to be consumed before we proceed down the stack
 		for (int32 KeyIndex = 0; KeyIndex < KeysToConsume.Num(); ++KeyIndex)
 		{
 			if (FKeyState* KeyState = KeyStateMap.Find(KeysToConsume[KeyIndex]))
@@ -873,6 +880,7 @@ bool UEnhancedPlayerInput::EvaluateInputComponentDelegates(UInputComponent* cons
 	}
 
 	KeysToConsume.Reset();
+	KeysToConsumeBeforeLegacy.Reset();
 
 	return bBlocksInput || bLegacyBlocksInput;
 }
