[Input] Fixed detection of sub-tick press-and-release during hitches

This happens due to Legacy input responding to sub-tick non-axis input taps differently to Enhanced Input. Enhanced Input has facility to detect these sub-tick taps, but only for axis-inputs that are treated as buttons (like analog triggers on a controller). This change adds a CVar for enabling the same detection for non-axis inputs.

Disable this change by setting the CVar EnhancedInput.DetectSubTickTapsForNonAxisInputs to false.

#tests Found a reliable repro for getting press and release events in the same tick, triggering the bug. Confirmed it no longer repros with the bug fix.

diff --git a/Engine/Plugins/EnhancedInput/Source/EnhancedInput/Private/EnhancedPlayerInput.cpp b/Engine/Plugins/EnhancedInput/Source/EnhancedInput/Private/EnhancedPlayerInput.cpp
--- a/Engine/Plugins/EnhancedInput/Source/EnhancedInput/Private/EnhancedPlayerInput.cpp
+++ b/Engine/Plugins/EnhancedInput/Source/EnhancedInput/Private/EnhancedPlayerInput.cpp
@@ -32,6 +32,11 @@ namespace UE::Input
 	static FAutoConsoleVariableRef CVarEnableListenerConsumption(TEXT("EnhancedInput.EnableListenerConsumption"),
 		bEnableListenerConsumption,
 		TEXT("If true, then bound listeners have the option to consume any Enhanced Input delegates on lower priority input components."));
+
+	static bool bDetectSubTickTapsForNonAxisInputs = true;
+	static FAutoConsoleVariableRef CVarDetectSubTickTapsForNonAxisInputs(TEXT("EnhancedInput.DetectSubTickTapsForNonAxisInputs"),
+		bDetectSubTickTapsForNonAxisInputs,
+		TEXT("If true, then vast press-release pairs that occur within a tick will trigger for non-axis inputs just as they already do for axis buttons."));
 }
 
 namespace UE::EnhancedInput
@@ -249,9 +254,16 @@ bool UEnhancedPlayerInput::InputKey(const FInputKeyEventArgs& Params)
 {
 	const bool bResult = Super::InputKey(Params);
 
-	if (Params.Key.IsButtonAxis() && Params.Event == IE_Pressed)
+	if (Params.Event == IE_Pressed)
 	{
-		KeysPressedThisTick.FindOrAdd(Params.Key, FVector(Params.AmountDepressed, 0.0, 0.0));
+		if (Params.Key.IsButtonAxis())
+		{
+			KeysPressedThisTick.FindOrAdd(Params.Key, FVector(Params.AmountDepressed, 0.0, 0.0));
+		}
+		else if (UE::Input::bDetectSubTickTapsForNonAxisInputs && Params.Key.IsDigital())
+		{
+			KeysPressedThisTick.FindOrAdd(Params.Key, FVector(1.0f, 0.0, 0.0));
+		}
 	}
 	
 	return bResult;
@@ -648,7 +660,7 @@ bool UEnhancedPlayerInput::EvaluateInputComponentDelegates(UInputComponent* cons
 							// Skip any action whose triggering keys have already been consumed.
 							if (FKeyState* KeyState = KeyStateMap.Find(KeyToConsume))
 							{
-								if (KeyState->bDown || KeyState->bDownPrevious)
+								if (KeyState->bDown || KeyState->bDownPrevious || KeyState->EventCounts[IE_Pressed].Num())
 								{
 									if (KeyState->bConsumed)
 									{
