[Enhanced Input] For the deadzone input modifier, when the upper and lower deadzone thresholds are equal, consider it an invalid asset setup. We do this because if you set them equal to each other, then there is no acceptable values the user could possibly input which would be accepted.

Account for this in PostEditChangeProperty to make it so that you always have a valid set up in editor now.

#rnx
#Synced-CL 48814820

diff --git a/Engine/Plugins/EnhancedInput/Source/EnhancedInput/Private/InputModifiers.cpp b/Engine/Plugins/EnhancedInput/Source/EnhancedInput/Private/InputModifiers.cpp
--- a/Engine/Plugins/EnhancedInput/Source/EnhancedInput/Private/InputModifiers.cpp
+++ b/Engine/Plugins/EnhancedInput/Source/EnhancedInput/Private/InputModifiers.cpp
@@ -129,10 +129,10 @@ EDataValidationResult UInputModifierDeadZone::IsDataValid(class FDataValidationC
 	EDataValidationResult Result = CombineDataValidationResults(Super::IsDataValid(Context), EDataValidationResult::Valid);
 
 	// You cannot scale a boolean value
-	if (LowerThreshold > UpperThreshold)
+	if (LowerThreshold >= UpperThreshold)
 	{
 		Result = EDataValidationResult::Invalid;
-		Context.AddError(LOCTEXT("InputModifierDeadZone", "The 'Lower Threshold' cannot be greater then the 'Upper Threshold' of a deadzone."));
+		Context.AddError(LOCTEXT("InputModifierDeadZone", "The 'Lower Threshold' cannot be greater than or equal to the 'Upper Threshold' of a deadzone."));
 	}
 
 	return Result;
@@ -149,7 +149,7 @@ void UInputModifierDeadZone::PostEditChangeProperty(FPropertyChangedEvent& Prope
 		// Clamp the lower threshold to the upper threshold value.
 		if (LowerThreshold > UpperThreshold)
 		{
-			LowerThreshold = UpperThreshold;
+			LowerThreshold = UpperThreshold - 0.001f;
 		}
 	}
 }
diff --git a/Engine/Plugins/EnhancedInput/Source/EnhancedInput/Public/InputModifiers.h b/Engine/Plugins/EnhancedInput/Source/EnhancedInput/Public/InputModifiers.h
--- a/Engine/Plugins/EnhancedInput/Source/EnhancedInput/Public/InputModifiers.h
+++ b/Engine/Plugins/EnhancedInput/Source/EnhancedInput/Public/InputModifiers.h
@@ -150,12 +150,19 @@ class UInputModifierDeadZone : public UInputModifier
 
 public:
 
-	// Threshold below which input is ignored
-	// This value should always be lower then the UpperThreshold.
+	/** 
+	* Threshold below which input is ignored.
+	*
+	* This value should always be less than the UpperThreshold.
+	*/
 	UPROPERTY(EditInstanceOnly, BlueprintReadWrite, Category=Settings, Config, meta=(ClampMin=0, ClampMax=1))
 	float LowerThreshold = 0.2f;
 
-	// Threshold above which input is clamped to 1
+	/** 
+	* Threshold above which input is clamped to 1.0.
+	* 
+	* This value should always be greater than the "Lower Threshold" value. 
+	*/
 	UPROPERTY(EditInstanceOnly, BlueprintReadWrite, Category=Settings, Config, meta=(ClampMin=0, ClampMax=1))
 	float UpperThreshold = 1.f;
 
