Enhanced Input no longer Ticks on Servers.

#rnx
#Synced-CL 48850057

diff --git a/Engine/Plugins/EnhancedInput/Source/EnhancedInput/Private/EnhancedInputModule.cpp b/Engine/Plugins/EnhancedInput/Source/EnhancedInput/Private/EnhancedInputModule.cpp
--- a/Engine/Plugins/EnhancedInput/Source/EnhancedInput/Private/EnhancedInputModule.cpp
+++ b/Engine/Plugins/EnhancedInput/Source/EnhancedInput/Private/EnhancedInputModule.cpp
@@ -36,6 +36,12 @@ public:
 
 	virtual ETickableTickType GetTickableTickType() const override
 	{
+		// Input can only possibly exist on the client, so never attempt to tick it on the server.
+		if (IsRunningDedicatedServer())
+		{
+			return ETickableTickType::Never;
+		}
+		
 		return ETickableTickType::Always;
 	}
 
