[Enhanced Input] When using ShouldCheckForLegacyKeyConsumption(), we should still be checking these against the individual action's bConsumeInput

Some Enhanced Input Actions opt OUT of consuming keys. These should also not consume legacy keys, so that behavior is consistent when migrating from legacy to enhanced.

#tests PIE
#Synced-CL 49015322

diff --git a/Engine/Plugins/EnhancedInput/Source/EnhancedInput/Private/EnhancedPlayerInput.cpp b/Engine/Plugins/EnhancedInput/Source/EnhancedInput/Private/EnhancedPlayerInput.cpp
--- a/Engine/Plugins/EnhancedInput/Source/EnhancedInput/Private/EnhancedPlayerInput.cpp
+++ b/Engine/Plugins/EnhancedInput/Source/EnhancedInput/Private/EnhancedPlayerInput.cpp
@@ -642,11 +642,12 @@ bool UEnhancedPlayerInput::EvaluateInputComponentDelegates(UInputComponent* cons
 		for (const TUniquePtr<FEnhancedInputActionEventBinding>& Binding : IC->GetActionEventBindings())
 		{
 			// PERF: Lots of map lookups! Group EnhancedActionBindings by Action?
-			if (const FInputActionInstance* ActionData = FindActionInstanceData(Binding->GetAction()))
+			const UInputAction* Action = Binding->GetAction();
+			if (const FInputActionInstance* ActionData = FindActionInstanceData(Action))
 			{
-				const FKeyConsumptionOptions* const ConsumptionData = KeyConsumptionData.Find(ActionData->GetSourceAction());
+				const FKeyConsumptionOptions* const ConsumptionData = KeyConsumptionData.Find(Action);
 
-				if (bCheckConsumedKeys && bShouldCheckForLegacyKeyConsumption && (ActionData->LastTriggerState != ETriggerState::None || ActionData->TriggerEvent != ETriggerEvent::None))
+				if (bCheckConsumedKeys && bShouldCheckForLegacyKeyConsumption && Action->bConsumeInput && (ActionData->LastTriggerState != ETriggerState::None || ActionData->TriggerEvent != ETriggerEvent::None))
 				{
 					// Don't do this action if its triggering keys were consumed elsewhere
 					if (ConsumptionData)
