[Enhanced Input] Make UEnhancedInputUserWidgetValidator not a BP type.

#rnx
#Synced-CL 49127187

diff --git a/Engine/Plugins/EnhancedInput/Source/InputBlueprintNodes/Private/EnhancedInputUserWidgetValidator.h b/Engine/Plugins/EnhancedInput/Source/InputBlueprintNodes/Private/EnhancedInputUserWidgetValidator.h
--- a/Engine/Plugins/EnhancedInput/Source/InputBlueprintNodes/Private/EnhancedInputUserWidgetValidator.h
+++ b/Engine/Plugins/EnhancedInput/Source/InputBlueprintNodes/Private/EnhancedInputUserWidgetValidator.h
@@ -17,7 +17,7 @@ class UWidgetBlueprint;
  * Widgets require bAutomaticallyRegisterInputOnConstruction to be true in order to
  * receive callbacks from Enhanced Input.
  */
-UCLASS()
+UCLASS(NotBlueprintType)
 class UEnhancedInputUserWidgetValidator : public UEditorValidatorBase
 {
 	GENERATED_BODY()
