[Enhanced Input] Deprecate the combo trigger. This trigger type is unreliable and can cause mapping contexts to be corrupted in the editor.

- Make UInputTriggerCombo deprecated and final
- FInputComboStepData deprecated
- FInputCancelAction deprecated
- Add asset validation warning about this deprecation.

#Synced-CL 49559950

diff --git a/Engine/Plugins/EnhancedInput/Source/EnhancedInput/Private/EnhancedActionKeyMapping.cpp b/Engine/Plugins/EnhancedInput/Source/EnhancedInput/Private/EnhancedActionKeyMapping.cpp
--- a/Engine/Plugins/EnhancedInput/Source/EnhancedInput/Private/EnhancedActionKeyMapping.cpp
+++ b/Engine/Plugins/EnhancedInput/Source/EnhancedInput/Private/EnhancedActionKeyMapping.cpp
@@ -123,8 +123,10 @@ EDataValidationResult FEnhancedActionKeyMapping::IsDataValid(FDataValidationCont
 	{
 		if (Trigger != nullptr)
 		{
+			PRAGMA_DISABLE_DEPRECATION_WARNINGS
 			// check if it the trigger is a combo or not
 			Trigger->IsA(UInputTriggerCombo::StaticClass()) ? bContextContainsComboTrigger = true : bContextContainsNonComboTrigger = true;
+			PRAGMA_ENABLE_DEPRECATION_WARNINGS
 			
 			Result = CombineDataValidationResults(Result, Trigger->IsDataValid(Context));
 		}
@@ -144,10 +146,12 @@ EDataValidationResult FEnhancedActionKeyMapping::IsDataValid(FDataValidationCont
 		// we also need to check the input action triggers for combo triggers
 		for (const UInputTrigger* Trigger : Action->Triggers)
 		{
+			PRAGMA_DISABLE_DEPRECATION_WARNINGS
 			if (Trigger != nullptr)
 			{
 				Trigger->IsA(UInputTriggerCombo::StaticClass()) ? bInputActionContainsComboTrigger = true : bInputActionContainsNonComboTrigger = true;
 			}
+			PRAGMA_ENABLE_DEPRECATION_WARNINGS
 		}
 
 		FFormatNamedArguments Args;
diff --git a/Engine/Plugins/EnhancedInput/Source/EnhancedInput/Private/EnhancedInputModule.cpp b/Engine/Plugins/EnhancedInput/Source/EnhancedInput/Private/EnhancedInputModule.cpp
--- a/Engine/Plugins/EnhancedInput/Source/EnhancedInput/Private/EnhancedInputModule.cpp
+++ b/Engine/Plugins/EnhancedInput/Source/EnhancedInput/Private/EnhancedInputModule.cpp
@@ -18,7 +18,9 @@
 
 #define LOCTEXT_NAMESPACE "EnhancedInput"
 
+PRAGMA_DISABLE_DEPRECATION_WARNINGS
 const FKey FEnhancedInputKeys::ComboKey("ComboKey");
+PRAGMA_ENABLE_DEPRECATION_WARNINGS
 
 DEFINE_LOG_CATEGORY(LogEnhancedInput);
 
@@ -235,11 +237,13 @@ void FEnhancedInputModule::StartupModule()
 {
 	Library = NewObject<UEnhancedInputLibrary>(GetTransientPackage(), UEnhancedInputLibrary::StaticClass(), NAME_None);
 	Library->AddToRoot();
-	
+
+	PRAGMA_DISABLE_DEPRECATION_WARNINGS
 	const FName NAME_EnhancedInput(TEXT("EnhancedInput"));
 	EKeys::AddMenuCategoryDisplayInfo(NAME_EnhancedInput, LOCTEXT("EnhancedInputSubCateogry", "Enhanced Input"), TEXT("GraphEditor.KeyEvent_16x"));
-	EKeys::AddKey(FKeyDetails(FEnhancedInputKeys::ComboKey, LOCTEXT("ComboKey", "Combo Key"), FKeyDetails::NotActionBindableKey, NAME_EnhancedInput));
-
+	EKeys::AddKey(FKeyDetails(FEnhancedInputKeys::ComboKey, LOCTEXT("ComboKey", "Combo Key (Deprecated)"), FKeyDetails::NotActionBindableKey, NAME_EnhancedInput));
+	PRAGMA_ENABLE_DEPRECATION_WARNINGS
+	
 	if (!IsRunningDedicatedServer())
 	{
 #if ENABLE_DRAW_DEBUG
diff --git a/Engine/Plugins/EnhancedInput/Source/EnhancedInput/Private/InputAction.cpp b/Engine/Plugins/EnhancedInput/Source/EnhancedInput/Private/InputAction.cpp
--- a/Engine/Plugins/EnhancedInput/Source/EnhancedInput/Private/InputAction.cpp
+++ b/Engine/Plugins/EnhancedInput/Source/EnhancedInput/Private/InputAction.cpp
@@ -174,9 +174,10 @@ EDataValidationResult UInputAction::IsDataValid(FDataValidationContext& Context)
 	{
 		if (Trigger)
 		{
+			PRAGMA_DISABLE_DEPRECATION_WARNINGS
 			// check if it the trigger is a combo or not
 			Trigger->IsA(UInputTriggerCombo::StaticClass()) ? bContainsComboTrigger = true : bContainsNonComboTrigger = true;
-			
+			PRAGMA_ENABLE_DEPRECATION_WARNINGS
 			Result = CombineDataValidationResults(Result, Trigger->IsDataValid(Context));
 		}
 		else
diff --git a/Engine/Plugins/EnhancedInput/Source/EnhancedInput/Private/InputTriggers.cpp b/Engine/Plugins/EnhancedInput/Source/EnhancedInput/Private/InputTriggers.cpp
--- a/Engine/Plugins/EnhancedInput/Source/EnhancedInput/Private/InputTriggers.cpp
+++ b/Engine/Plugins/EnhancedInput/Source/EnhancedInput/Private/InputTriggers.cpp
@@ -342,6 +342,8 @@ ETriggerState UInputTriggerChordAction::UpdateState_Implementation(const UEnhanc
 	return EventData ? EventData->GetEvaluatedActionTriggerState() : ETriggerState::None;
 }
 
+PRAGMA_DISABLE_DEPRECATION_WARNINGS
+
 UInputTriggerCombo::UInputTriggerCombo()
 {
 	bShouldAlwaysTick = true;
@@ -440,6 +442,8 @@ ETriggerState UInputTriggerCombo::UpdateState_Implementation(const UEnhancedPlay
 EDataValidationResult UInputTriggerCombo::IsDataValid(FDataValidationContext& Context) const
 {
 	EDataValidationResult Result = CombineDataValidationResults(Super::IsDataValid(Context), EDataValidationResult::Valid);
+
+	Context.AddWarning(LOCTEXT("ComboTriggerDeprecated", "The input combo trigger is deprecated and will be removed in a future release"));
 	
 	// You can't evaluate the combo if there are no combo steps!
 	if (ComboActions.IsEmpty())
@@ -472,4 +476,6 @@ EDataValidationResult UInputTriggerCombo::IsDataValid(FDataValidationContext& Co
 }
 #endif // WITH_EDITOR
 
+PRAGMA_ENABLE_DEPRECATION_WARNINGS
+
 #undef LOCTEXT_NAMESPACE
diff --git a/Engine/Plugins/EnhancedInput/Source/EnhancedInput/Public/EnhancedInputModule.h b/Engine/Plugins/EnhancedInput/Source/EnhancedInput/Public/EnhancedInputModule.h
--- a/Engine/Plugins/EnhancedInput/Source/EnhancedInput/Public/EnhancedInputModule.h
+++ b/Engine/Plugins/EnhancedInput/Source/EnhancedInput/Public/EnhancedInputModule.h
@@ -17,6 +17,7 @@ ENHANCEDINPUT_API DECLARE_LOG_CATEGORY_EXTERN(LogEnhancedInput, Log, All);
 struct FEnhancedInputKeys
 {
 	// Combo FKey that serves as the key combo triggers are automatically mapped to - is not action bindable
+	UE_DEPRECATED(5.8, "The combo trigger and key have been deprecated")
 	static UE_API const FKey ComboKey;
 };
 
diff --git a/Engine/Plugins/EnhancedInput/Source/EnhancedInput/Public/InputTriggers.h b/Engine/Plugins/EnhancedInput/Source/EnhancedInput/Public/InputTriggers.h
--- a/Engine/Plugins/EnhancedInput/Source/EnhancedInput/Public/InputTriggers.h
+++ b/Engine/Plugins/EnhancedInput/Source/EnhancedInput/Public/InputTriggers.h
@@ -486,7 +486,7 @@ protected:
 };
 
 USTRUCT(BlueprintType)
-struct FInputComboStepData
+struct UE_DEPRECATED(5.8, "The Input combo trigger and related types are deprecated") FInputComboStepData
 {
 	GENERATED_BODY()
 
@@ -507,7 +507,7 @@ struct FInputComboStepData
 };
 
 USTRUCT(BlueprintType)
-struct FInputCancelAction
+struct UE_DEPRECATED(5.8, "The Input combo trigger and related types are deprecated") FInputCancelAction
 {
 	GENERATED_BODY()
 
@@ -527,7 +527,7 @@ struct FInputCancelAction
  * Note: This will only trigger for one frame before resetting the combo trigger's progress 
  */
 UCLASS(MinimalAPI, NotBlueprintable, meta = (DisplayName = "Combo (Beta)", NotInputConfigurable = "true"))
-class UInputTriggerCombo : public UInputTrigger
+class UE_DEPRECATED(5.8, "UInputTriggerCombo has been deprecated.") UInputTriggerCombo final : public UInputTrigger
 {
 	GENERATED_BODY()
 
@@ -552,6 +552,8 @@ public:
 #if WITH_EDITOR
 	UE_API virtual EDataValidationResult IsDataValid(class FDataValidationContext& Context) const override;
 #endif
+
+PRAGMA_DISABLE_DEPRECATION_WARNINGS
 	
 	/**
 	 * List of input actions that need to be completed (according to Combo Step Completion States) to trigger this action.
@@ -564,6 +566,8 @@ public:
 	UPROPERTY(EditAnywhere, BlueprintReadWrite, Category = "Trigger Settings", meta = (DisplayThumbnail = "false", DisplayName = "Cancel Actions"))
 	TArray<FInputCancelAction> InputCancelActions;
 
+PRAGMA_ENABLE_DEPRECATION_WARNINGS
+	
 	/** Determines what kind of trigger events can happen from the behavior of this trigger. */
 	virtual ETriggerEventsSupported GetSupportedTriggerEvents() const override { return ETriggerEventsSupported::All; }
 };
diff --git a/Engine/Plugins/EnhancedInput/Source/InputEditor/Private/InputCustomizations.cpp b/Engine/Plugins/EnhancedInput/Source/InputEditor/Private/InputCustomizations.cpp
--- a/Engine/Plugins/EnhancedInput/Source/InputEditor/Private/InputCustomizations.cpp
+++ b/Engine/Plugins/EnhancedInput/Source/InputEditor/Private/InputCustomizations.cpp
@@ -61,7 +61,9 @@ void FEnhancedActionMappingCustomization::CustomizeHeader(TSharedRef<IPropertyHa
 
 	// TODO: Use FDetailArrayBuilder?
 
+	PRAGMA_DISABLE_DEPRECATION_WARNINGS
 	const bool bContainsComboTrigger = DoesTriggerArrayContainCombo();
+	PRAGMA_ENABLE_DEPRECATION_WARNINGS
 	// Pass our header row into the key struct customizeheader method so it populates our row with the key struct header
 	KeyStructCustomization = StaticCastSharedPtr<FKeyStructCustomization>(KeyStructInstance);
 	KeyStructCustomization->SetDefaultKeyName("ComboKey");
@@ -150,7 +152,9 @@ void FEnhancedActionMappingCustomization::RemoveMappingButton_OnClick() const
 
 void FEnhancedActionMappingCustomization::OnTriggersChanged() const
 {
+	PRAGMA_DISABLE_DEPRECATION_WARNINGS
 	const bool bContainsComboTrigger = DoesTriggerArrayContainCombo();
+	PRAGMA_ENABLE_DEPRECATION_WARNINGS
 	
 	// if it is currently disabled and doesn't contain a combo now we should set the key to something other than ComboKey
 	if (!KeyStructCustomization->GetEnableKeySelector() && !bContainsComboTrigger)
@@ -204,6 +208,7 @@ void FEnhancedActionMappingCustomization::OnInputActionModifiersChanged() const
     }
 }
 
+PRAGMA_DISABLE_DEPRECATION_WARNINGS
 bool FEnhancedActionMappingCustomization::DoesTriggerArrayContainCombo() const
 {
 	if (const FEnhancedActionKeyMapping* ActionKeyMapping = GetActionKeyMapping(MappingPropertyHandle))
@@ -231,6 +236,7 @@ bool FEnhancedActionMappingCustomization::DoesTriggerArrayContainCombo() const
 	
 	return false;
 }
+PRAGMA_ENABLE_DEPRECATION_WARNINGS
 
 //////////////////////////////////////////////////////////
 // FEnhancedInputDeveloperSettingsCustomization
diff --git a/Engine/Plugins/EnhancedInput/Source/InputEditor/Private/InputCustomizations.h b/Engine/Plugins/EnhancedInput/Source/InputEditor/Private/InputCustomizations.h
--- a/Engine/Plugins/EnhancedInput/Source/InputEditor/Private/InputCustomizations.h
+++ b/Engine/Plugins/EnhancedInput/Source/InputEditor/Private/InputCustomizations.h
@@ -24,6 +24,7 @@ public:
 	virtual void CustomizeChildren(TSharedRef<IPropertyHandle> PropertyHandle, IDetailChildrenBuilder& ChildBuilder, IPropertyTypeCustomizationUtils& CustomizationUtils) override;
 
 	/** Searches the trigger array of an input action - returns true if any combo input triggers are present */
+	UE_DEPRECATED(5.8, "The combo trigger is now deprecated")
 	bool DoesTriggerArrayContainCombo() const;
 private:
 
