[Enhanced Input] Fix mappings to AnyKey getting into a broken state when mouse scroll wheel up/down were used.

This bug was happening because the Mouse Scroll buttons (EKeys::MouseScrollUp and EKeys::MouseScrollDown) are "Digital Axis" keys, which Enhanced Input was not correctly determining the state of. The scroll up/down buttons are pressed and released on the same frame, and will not have an IE_Released event like normal digital buttons. This resulted in the AnyKey logic only running every other press, so you would have to scroll up twice to press & release.

This is now fixed by correctly accounting for the DigitalAxis keys in enhanced input and being more robust in our checks for if "Any Key" was pressed.

#tests EI unit tests, play in QA Game and Lyra

#Synced-CL 50164937

diff --git a/Engine/Plugins/EnhancedInput/Source/EnhancedInput/Private/EnhancedPlayerInput.cpp b/Engine/Plugins/EnhancedInput/Source/EnhancedInput/Private/EnhancedPlayerInput.cpp
--- a/Engine/Plugins/EnhancedInput/Source/EnhancedInput/Private/EnhancedPlayerInput.cpp
+++ b/Engine/Plugins/EnhancedInput/Source/EnhancedInput/Private/EnhancedPlayerInput.cpp
@@ -37,6 +37,37 @@ namespace UE::Input
 	static FAutoConsoleVariableRef CVarDetectSubTickTapsForNonAxisInputs(TEXT("EnhancedInput.DetectSubTickTapsForNonAxisInputs"),
 		bDetectSubTickTapsForNonAxisInputs,
 		TEXT("If true, then vast press-release pairs that occur within a tick will trigger for non-axis inputs just as they already do for axis buttons."));
+
+	static bool IsKeyValidForAnyKey(const FKey& Key)
+	{
+		return !Key.IsAnalog() && (Key.IsDigital() || Key.IsButtonAxis());
+	}
+
+	static bool IsAnyKeyStateConsideredPressed(const FKey& Key, const FKeyState& State)
+	{
+		if (!IsKeyValidForAnyKey(Key))
+		{
+			return false;
+		}
+
+		const bool bHasAnyPressedEvents = State.EventCounts[IE_Pressed].Num() > 0 || State.EventAccumulator[IE_Pressed].Num() > 0;
+
+		// Button axis keys may not have the IE_Pressed event set, so check their raw value as well
+		return bHasAnyPressedEvents || (Key.IsButtonAxis() && State.RawValue.Size() > 0.0);
+	}
+
+	static bool IsAnyKeyStateConsideredReleased(const FKey& Key, const FKeyState& State)
+	{
+		if (!IsKeyValidForAnyKey(Key))
+		{
+			return false;
+		}
+
+		const bool bHasAnyReleasedEvents = State.EventCounts[IE_Released].Num() > 0 || State.EventAccumulator[IE_Released].Num() > 0;
+
+		// Button axis keys may not have the IE_Released event set, so consider them 'released' if the value is zero
+		return bHasAnyReleasedEvents || (Key.IsButtonAxis() && State.RawValue.IsNearlyZero());
+	}
 }
 
 namespace UE::EnhancedInput
@@ -323,7 +354,10 @@ void UEnhancedPlayerInput::EvaluateKeyMapState(const float DeltaTime, const bool
 		const FKeyState& KeyState = KeyPair.Value;
 		// TODO: Can't just use bDownPrevious as paired axis event edges may not fire due to axial deadzoning/missing axis properties. Need to change how this is detected in PlayerInput.cpp.
 		bool bWasDown = KeyState.bDownPrevious || KeyState.EventCounts[IE_Pressed].Num() || KeyState.EventCounts[IE_Repeat].Num();
-		bWasDown |= KeyPair.Key.IsAnalog() && KeyState.RawValue.SizeSquared() != 0;	// Analog inputs should pulse every (non-zero) tick to retain compatibility with UE4.
+
+		bWasDown |=
+			(KeyPair.Key.IsAnalog() || KeyPair.Key.IsButtonAxis()) &&
+			KeyState.RawValue.SizeSquared() != 0;	// Analog inputs should pulse every (non-zero) tick to retain compatibility with UE4.
 
 		// When UPlayerInput::FlushPressedKeys is called any keys that are down will have their RawValue set to 0, and their bDown/bDownPrevious state will be reset to false.
 		// However, their "Value" will not be reset until UPlayerInput::ProcessInputStack. We need to detect if this key was down previously after a flush
@@ -338,7 +372,11 @@ void UEnhancedPlayerInput::EvaluateKeyMapState(const float DeltaTime, const bool
 		bWasDown |= bKeyWasJustFlushed;
 			
 		// Keep track of the state of every key so that when we are done iterating we can have a meaningful value for EKeys::AnyKey
-		bWasAnyKeyDownLastFrame |= bWasDown;
+		// Any key does not support analog events, and we want to know if any non-analog key was pressed last frame at all.
+		if (UE::Input::IsAnyKeyStateConsideredPressed(KeyPair.Key, KeyState))
+		{
+			bWasAnyKeyDownLastFrame |= bWasDown;
+		}
 
 		KeyDownPrevious.Emplace(KeyPair.Key, bWasDown);
 	}
@@ -437,13 +475,13 @@ void UEnhancedPlayerInput::PrepareInputDelegatesForEvaluation(const TArray<UInpu
 			for (TPair<FKey, FKeyState>& KeyStatePair : KeyStateMap)
 			{
 				// EKeys::AnyKey will only use non-analog keys. the same as the legacy system.
-				if (!KeyStatePair.Key.IsDigital())
+				if (!UE::Input::IsKeyValidForAnyKey(KeyStatePair.Key))
 				{
 					continue;
 				}
 
 				// If we have no Substitute key, then we can simply use the key state of the first available key with pressed events.
-				if (CurrentlyInUseAnyKeySubstitute == NAME_None && KeyStatePair.Value.EventCounts[IE_Pressed].Num() > 0)
+				if (CurrentlyInUseAnyKeySubstitute == NAME_None && UE::Input::IsAnyKeyStateConsideredPressed(KeyStatePair.Key, KeyStatePair.Value))
 				{
 					KeyState = &KeyStatePair.Value;
 					CurrentlyInUseAnyKeySubstitute = KeyStatePair.Key.GetFName();
@@ -455,7 +493,7 @@ void UEnhancedPlayerInput::PrepareInputDelegatesForEvaluation(const TArray<UInpu
 				{
 					// If the substitute key was just released, then we can reset our currently in 
 					// use substitute key so that it can be replaced by something else on the next key press.					
-					if (KeyStatePair.Value.EventCounts[IE_Released].Num() > 0)
+					if (UE::Input::IsAnyKeyStateConsideredReleased(KeyStatePair.Key, KeyStatePair.Value))
 					{
 						CurrentlyInUseAnyKeySubstitute = NAME_None;
 					}
@@ -498,14 +536,16 @@ void UEnhancedPlayerInput::PrepareInputDelegatesForEvaluation(const TArray<UInpu
 		// TODO: Can't just use bDown as paired axis event edges may not fire due to axial deadzoning/missing axis properties. Need to change how this is detected in PlayerInput.cpp.
 		bool bKeyIsDown = KeyState && (KeyState->bDown || KeyState->EventCounts[IE_Pressed].Num() || KeyState->EventCounts[IE_Repeat].Num());
 		// Analog inputs should pulse every (non-zero) tick to retain compatibility with UE4. TODO: This would be better handled at the device level.
-		bKeyIsDown |= Mapping.Key.IsAnalog() && RawKeyValue.SizeSquared() > 0;
+		bKeyIsDown |= (Mapping.Key.IsAnalog() || Mapping.Key.IsButtonAxis()) && RawKeyValue.SizeSquared() > 0;
 
-		bool bKeyIsReleased = !bKeyIsDown && bDownLastTick;
-		bool bKeyIsHeld = bKeyIsDown && bDownLastTick;
+		const bool bKeyIsReleased = !bKeyIsDown && bDownLastTick;
+		const bool bKeyIsHeld = bKeyIsDown && bDownLastTick;
 
-		EKeyEvent KeyEvent = bKeyIsHeld ? EKeyEvent::Held : ((bKeyIsDown || bKeyIsReleased) ? EKeyEvent::Actuated : EKeyEvent::None);
+		const EKeyEvent KeyEvent = bKeyIsHeld ? EKeyEvent::Held : ((bKeyIsDown || bKeyIsReleased) ? EKeyEvent::Actuated : EKeyEvent::None);
 
-		FVector* PressedThisTickValue = KeysPressedThisTick.Find(Mapping.Key);
+		const FVector* PressedThisTickValue = (Mapping.Key == EKeys::AnyKey) ?
+			KeysPressedThisTick.Find(FKey(CurrentlyInUseAnyKeySubstitute)) :
+			KeysPressedThisTick.Find(Mapping.Key);
 		
 		// For keys that were pressed and released within the same frame, set their RawValue so that
 		// InputTriggers are aware that they have been pressed
