PR #14328: Fixed virtual key mappings in enhanced input not correctly processing held and released events

#rnx
#Synced-CL 51053450

diff --git a/Engine/Plugins/EnhancedInput/Source/EnhancedInput/Private/EnhancedPlayerInput.cpp b/Engine/Plugins/EnhancedInput/Source/EnhancedInput/Private/EnhancedPlayerInput.cpp
--- a/Engine/Plugins/EnhancedInput/Source/EnhancedInput/Private/EnhancedPlayerInput.cpp
+++ b/Engine/Plugins/EnhancedInput/Source/EnhancedInput/Private/EnhancedPlayerInput.cpp
@@ -532,7 +532,7 @@ void UEnhancedPlayerInput::PrepareInputDelegatesForEvaluation(const TArray<UInpu
 		}
 		
 		// Establish update type.
-		bool bDownLastTick = KeyDownPrevious.FindRef(Mapping.Key);
+		bool bDownLastTick = KeyDownPrevious.FindRef(Mapping.Key.GetVirtualKey());
 		// TODO: Can't just use bDown as paired axis event edges may not fire due to axial deadzoning/missing axis properties. Need to change how this is detected in PlayerInput.cpp.
 		bool bKeyIsDown = KeyState && (KeyState->bDown || KeyState->EventCounts[IE_Pressed].Num() || KeyState->EventCounts[IE_Repeat].Num());
 		// Analog inputs should pulse every (non-zero) tick to retain compatibility with UE4. TODO: This would be better handled at the device level.
