PR #14327: Fixed virtual keys not working with enhanced input mappings in CommonUI

#rnx

#Synced-CL 51053450

diff --git a/Engine/Plugins/EnhancedInput/Source/EnhancedInput/Private/EnhancedInputSubsystemInterface.cpp b/Engine/Plugins/EnhancedInput/Source/EnhancedInput/Private/EnhancedInputSubsystemInterface.cpp
--- a/Engine/Plugins/EnhancedInput/Source/EnhancedInput/Private/EnhancedInputSubsystemInterface.cpp
+++ b/Engine/Plugins/EnhancedInput/Source/EnhancedInput/Private/EnhancedInputSubsystemInterface.cpp
@@ -722,7 +722,7 @@ TArray<FKey> IEnhancedInputSubsystemInterface::QueryKeysMappedToAction(const UIn
 			{
 				if (Mapping.Action == Action)
 				{
-					MappedKeys.AddUnique(Mapping.Key);
+					MappedKeys.AddUnique(Mapping.Key.GetVirtualKey());
 				}
 			}
 		}
