[Enhanced Input] Replace uses of UE_LOG with UE_LOGF

#rnx

#Synced-CL 51092797

diff --git a/Engine/Plugins/EnhancedInput/Source/EnhancedInput/Private/EnhancedInputLibrary.cpp b/Engine/Plugins/EnhancedInput/Source/EnhancedInput/Private/EnhancedInputLibrary.cpp
--- a/Engine/Plugins/EnhancedInput/Source/EnhancedInput/Private/EnhancedInputLibrary.cpp
+++ b/Engine/Plugins/EnhancedInput/Source/EnhancedInput/Private/EnhancedInputLibrary.cpp
@@ -77,12 +77,12 @@ FInputActionValue UEnhancedInputLibrary::GetBoundActionValue(AActor* Actor, cons
 	{
 		if (!Actor)
 		{
-			UE_LOG(LogEnhancedInput, Error, TEXT("UEnhancedInputLibrary::GetBoundActionValue was called with an invalid Actor!"));
+			UE_LOGF(LogEnhancedInput, Error, "UEnhancedInputLibrary::GetBoundActionValue was called with an invalid Actor!");
 		}
 		
 		if (!Action)
 		{
-			UE_LOG(LogEnhancedInput, Error, TEXT("UEnhancedInputLibrary::GetBoundActionValue was called with an invalid Action!"));
+			UE_LOGF(LogEnhancedInput, Error, "UEnhancedInputLibrary::GetBoundActionValue was called with an invalid Action!");
 		}		
 		ensureMsgf(false, TEXT("Invalid GetBoundActionValue call. Check logs for details!"));
 	}
diff --git a/Engine/Plugins/EnhancedInput/Source/EnhancedInput/Private/EnhancedInputSubsystemInterface.cpp b/Engine/Plugins/EnhancedInput/Source/EnhancedInput/Private/EnhancedInputSubsystemInterface.cpp
--- a/Engine/Plugins/EnhancedInput/Source/EnhancedInput/Private/EnhancedInputSubsystemInterface.cpp
+++ b/Engine/Plugins/EnhancedInput/Source/EnhancedInput/Private/EnhancedInputSubsystemInterface.cpp
@@ -95,7 +95,7 @@ void IEnhancedInputSubsystemInterface::BindUserSettingDelegates()
 	UEnhancedInputUserSettings* Settings = GetUserSettings();
 	if (!Settings)
 	{
-		UE_LOG(LogEnhancedInput, Error, TEXT("Unable to get the user settings object!"));
+		UE_LOGF(LogEnhancedInput, Error, "Unable to get the user settings object!");
 		return;
 	}
 
@@ -103,7 +103,7 @@ void IEnhancedInputSubsystemInterface::BindUserSettingDelegates()
 	// but do this in case someone implements this interface
 	if (!GetDefault<UEnhancedInputDeveloperSettings>()->bEnableUserSettings)
 	{
-		UE_LOG(LogEnhancedInput, Error, TEXT("Attempting to bind to user settings delegates but they are disabled in UEnhancedInputDeveloperSettings!"));
+		UE_LOGF(LogEnhancedInput, Error, "Attempting to bind to user settings delegates but they are disabled in UEnhancedInputDeveloperSettings!");
 		return;
 	}
 
@@ -161,12 +161,12 @@ void IEnhancedInputSubsystemInterface::StartContinuousInputInjectionForPlayerMap
 		}
 		else
 		{
-			UE_LOG(LogEnhancedInput, Warning, TEXT("Could not find a Input Action for mapping name '%s'"), *MappingName.ToString());
+			UE_LOGF(LogEnhancedInput, Warning, "Could not find a Input Action for mapping name '%ls'", *MappingName.ToString());
 		}
 	}
 	else
 	{
-		UE_LOG(LogEnhancedInput, Warning, TEXT("Could not find a valid UEnhancedInputUserSettings object, is it enabled in the project settings?"));
+		UE_LOGF(LogEnhancedInput, Warning, "Could not find a valid UEnhancedInputUserSettings object, is it enabled in the project settings?");
 	}
 }
 
@@ -188,12 +188,12 @@ void IEnhancedInputSubsystemInterface::UpdateValueOfContinuousInputInjectionForP
 		}
 		else
 		{
-			UE_LOG(LogEnhancedInput, Warning, TEXT("Could not find a Input Action for mapping name '%s'"), *MappingName.ToString());
+			UE_LOGF(LogEnhancedInput, Warning, "Could not find a Input Action for mapping name '%ls'", *MappingName.ToString());
 		}
 	}
 	else
 	{
-		UE_LOG(LogEnhancedInput, Warning, TEXT("Could not find a valid UEnhancedInputUserSettings object, is it enabled in the project settings?"));
+		UE_LOGF(LogEnhancedInput, Warning, "Could not find a valid UEnhancedInputUserSettings object, is it enabled in the project settings?");
 	}
 }
 
@@ -212,12 +212,12 @@ void IEnhancedInputSubsystemInterface::StopContinuousInputInjectionForPlayerMapp
 		}
 		else
 		{
-			UE_LOG(LogEnhancedInput, Warning, TEXT("Could not find a Input Action for mapping name '%s'"), *MappingName.ToString());
+			UE_LOGF(LogEnhancedInput, Warning, "Could not find a Input Action for mapping name '%ls'", *MappingName.ToString());
 		}
 	}
 	else
 	{
-		UE_LOG(LogEnhancedInput, Warning, TEXT("Could not find a valid UEnhancedInputUserSettings object, is it enabled in the project settings?"));
+		UE_LOGF(LogEnhancedInput, Warning, "Could not find a valid UEnhancedInputUserSettings object, is it enabled in the project settings?");
 	}
 }
 
@@ -237,12 +237,12 @@ void IEnhancedInputSubsystemInterface::InjectInputVectorForPlayerMapping(const F
 		}
 		else
 		{
-			UE_LOG(LogEnhancedInput, Warning, TEXT("Could not find a Input Action for mapping name '%s'"), *MappingName.ToString());
+			UE_LOGF(LogEnhancedInput, Warning, "Could not find a Input Action for mapping name '%ls'", *MappingName.ToString());
 		}
 	}
 	else
 	{
-		UE_LOG(LogEnhancedInput, Warning, TEXT("Could not find a valid UEnhancedInputUserSettings object, is it enabled in the project settings?"));
+		UE_LOGF(LogEnhancedInput, Warning, "Could not find a valid UEnhancedInputUserSettings object, is it enabled in the project settings?");
 	}
 }
 
@@ -297,7 +297,7 @@ void IEnhancedInputSubsystemInterface::AddMappingContext(const UInputMappingCont
 	}
 	else
 	{
-		UE_LOG(LogEnhancedInput, Warning, TEXT("Called AddMappingContext with a null Mapping Context! No changes have been applied."));
+		UE_LOGF(LogEnhancedInput, Warning, "Called AddMappingContext with a null Mapping Context! No changes have been applied.");
 	}
 }
 
@@ -356,7 +356,7 @@ FGameplayTagContainer IEnhancedInputSubsystemInterface::GetInputMode() const
 		return Input->GetCurrentInputMode();
 	}
 
-	UE_LOG(LogEnhancedInput, Error, TEXT("[%hs] Null player input, cannot get the current input mode."), __func__);
+	UE_LOGF(LogEnhancedInput, Error, "[%hs] Null player input, cannot get the current input mode.", __func__);
 	
 	return FGameplayTagContainer {};
 }
@@ -370,7 +370,7 @@ void IEnhancedInputSubsystemInterface::SetInputMode(const FGameplayTagContainer&
 	UEnhancedPlayerInput* Input = GetPlayerInput();
 	if (!Input)
 	{
-		UE_LOG(LogEnhancedInput, Error, TEXT("[%hs] Null player input, unable to set the input mode to '%s'"), __func__, *NewMode.ToString());
+		UE_LOGF(LogEnhancedInput, Error, "[%hs] Null player input, unable to set the input mode to '%ls'", __func__, *NewMode.ToString());
 		return;
 	}
 	
@@ -388,7 +388,7 @@ void IEnhancedInputSubsystemInterface::AppendTagsToInputMode(const FGameplayTagC
 	UEnhancedPlayerInput* Input = GetPlayerInput();
 	if (!Input)
 	{
-		UE_LOG(LogEnhancedInput, Error, TEXT("[%hs] Null player input, unable to append tags '%s' to the input mode"), __func__, *TagsToAdd.ToString());
+		UE_LOGF(LogEnhancedInput, Error, "[%hs] Null player input, unable to append tags '%ls' to the input mode", __func__, *TagsToAdd.ToString());
 		return;
 	}
 	
@@ -411,7 +411,7 @@ void IEnhancedInputSubsystemInterface::RemoveTagsFromInputMode(const FGameplayTa
 	UEnhancedPlayerInput* Input = GetPlayerInput();
 	if (!Input)
 	{
-		UE_LOG(LogEnhancedInput, Error, TEXT("[%hs] Null player input, unable to remove tags '%s' from the input mode"), __func__, *TagsToRemove.ToString());
+		UE_LOGF(LogEnhancedInput, Error, "[%hs] Null player input, unable to remove tags '%ls' from the input mode", __func__, *TagsToRemove.ToString());
 		return;
 	}
 
@@ -757,7 +757,7 @@ void IEnhancedInputSubsystemInterface::ValidateTrackedMappingContextsAreUnregist
 		{
 			if (IsValid(Pair.Key) && Pair.Key->GetRegistrationTrackingMode() != EMappingContextRegistrationTrackingMode::Untracked)
 			{
-				UE_LOG(LogEnhancedInput, Warning, TEXT("Input Mapping Context [%s] has tracking mode [%s] but is still applied and might be leaking, unregister it before deinitialization."),
+				UE_LOGF(LogEnhancedInput, Warning, "Input Mapping Context [%ls] has tracking mode [%ls] but is still applied and might be leaking, unregister it before deinitialization.",
 					*Pair.Key->GetName(), *UEnum::GetValueAsString(Pair.Key->GetRegistrationTrackingMode()));
 			}
 		}
@@ -791,7 +791,7 @@ TArray<FEnhancedActionKeyMapping> IEnhancedInputSubsystemInterface::ReorderMappi
 			{
 				if (!Trigger)
 				{
-					UE_LOG(LogEnhancedInput, Error, TEXT("Null input trigger detected in mapping to input action '%s'"), *GetNameSafe(Mapping.Action));
+					UE_LOGF(LogEnhancedInput, Error, "Null input trigger detected in mapping to input action '%ls'", *GetNameSafe(Mapping.Action));
 					continue;
 				}
 				
@@ -847,7 +847,7 @@ TArray<FEnhancedActionKeyMapping> IEnhancedInputSubsystemInterface::ReorderMappi
 		else
 		{
 			++NumEmptyMappings;
-			UE_LOG(LogEnhancedInput, Warning, TEXT("A Key Mapping with a blank action has been added! Ignoring the key mapping to '%s'"), *Mapping.Key.ToString());
+			UE_LOGF(LogEnhancedInput, Warning, "A Key Mapping with a blank action has been added! Ignoring the key mapping to '%ls'", *Mapping.Key.ToString());
 		}
 	}
 
@@ -916,7 +916,7 @@ void IEnhancedInputSubsystemInterface::RebuildControlMappings()
 		{
 			if (!Pair.Key || !Pair.Value)
 			{
-				UE_LOG(LogEnhancedInput, Error, TEXT("An invalid Mapping Context Redirect specified in '%s'"), PlatformSettings->GetConfigOverridePlatform());
+				UE_LOGF(LogEnhancedInput, Error, "An invalid Mapping Context Redirect specified in '%ls'", PlatformSettings->GetConfigOverridePlatform());
 				continue;
 			}
 			
@@ -930,10 +930,10 @@ void IEnhancedInputSubsystemInterface::RebuildControlMappings()
 				// Optional logging that may be helpful for debugging purposes
 				if (PlatformSettings->ShouldLogMappingContextRedirects())
 				{
-					UE_LOG(
+					UE_LOGF(
 						LogEnhancedInput,
 						Log,
-						TEXT("'%s' Redirecting Mapping Context '%s' -> '%s'"),
+						"'%ls' Redirecting Mapping Context '%ls' -> '%ls'",
 						PlatformSettings->GetConfigOverridePlatform(), *Pair.Key->GetName(), *Pair.Value->GetName()
 					);
 				}
@@ -967,7 +967,7 @@ void IEnhancedInputSubsystemInterface::RebuildControlMappings()
 
 		if (!ensure(IsValid(MappingContext)))
 		{
-			UE_LOG(LogEnhancedInput, Error, TEXT("[%hs] Encountered an invalid mapping context while rebuilding control mappings. Some inputs may be ignored!"), __func__);
+			UE_LOGF(LogEnhancedInput, Error, "[%hs] Encountered an invalid mapping context while rebuilding control mappings. Some inputs may be ignored!", __func__);
 			continue;
 		}
 		
@@ -977,7 +977,7 @@ void IEnhancedInputSubsystemInterface::RebuildControlMappings()
 			MappingContext->ShouldFilterMappingByInputMode() &&
 			!MappingContext->GetInputModeQuery().Matches(CurrentInputMode))
 		{
-			UE_LOG(LogEnhancedInput, Verbose, TEXT("[%hs] Not applying mappings from IMC '%s' because it does not meet the requirements of the current input mode '%s'"),
+			UE_LOGF(LogEnhancedInput, Verbose, "[%hs] Not applying mappings from IMC '%ls' because it does not meet the requirements of the current input mode '%ls'",
 				__func__, *GetNameSafe(MappingContext), *CurrentInputMode.ToString());
 				
 			continue;
diff --git a/Engine/Plugins/EnhancedInput/Source/EnhancedInput/Private/EnhancedInputSubsystems.cpp b/Engine/Plugins/EnhancedInput/Source/EnhancedInput/Private/EnhancedInputSubsystems.cpp
--- a/Engine/Plugins/EnhancedInput/Source/EnhancedInput/Private/EnhancedInputSubsystems.cpp
+++ b/Engine/Plugins/EnhancedInput/Source/EnhancedInput/Private/EnhancedInputSubsystems.cpp
@@ -94,7 +94,7 @@ void UEnhancedInputLocalPlayerSubsystem::InitalizeUserSettings()
 
 	if (!LP)
 	{
-		UE_LOG(LogEnhancedInput, Error, TEXT("Failed to initialize user settings, there is no valid local player!"));
+		UE_LOGF(LogEnhancedInput, Error, "Failed to initialize user settings, there is no valid local player!");
 		return;
 	}
 	
@@ -108,14 +108,14 @@ void UEnhancedInputLocalPlayerSubsystem::InitalizeUserSettings()
 
 	if (!GetDefault<UEnhancedInputDeveloperSettings>()->bEnableUserSettings)
 	{
-		UE_LOG(LogEnhancedInput, Verbose, TEXT("bEnableUserSettings is set to false, skipping creation of UEnhancedInputUserSettings!"));
+		UE_LOGF(LogEnhancedInput, Verbose, "bEnableUserSettings is set to false, skipping creation of UEnhancedInputUserSettings!");
 		return;
 	}
 
 	// We can't load the user settings without a valid player input object because we get the save file name from it
 	if (!GetPlayerInput())
 	{
-		UE_LOG(LogEnhancedInput, Warning, TEXT("UEnhancedInputLocalPlayerSubsystem for Local Player '%s' does not have a valid PlayerInput object. Failed to load user settings."), *GetNameSafe(LP));
+		UE_LOGF(LogEnhancedInput, Warning, "UEnhancedInputLocalPlayerSubsystem for Local Player '%ls' does not have a valid PlayerInput object. Failed to load user settings.", *GetNameSafe(LP));
 		return;
 	}
 	
@@ -128,7 +128,7 @@ void UEnhancedInputLocalPlayerSubsystem::InitalizeUserSettings()
 		
 		OnPostUserSettingsInitialized.Broadcast(UserSettings);
 		
-		UE_LOG(LogEnhancedInput, Log, TEXT("Enhanced Input local player subsystem has initialized the user settings!"));
+		UE_LOGF(LogEnhancedInput, Log, "Enhanced Input local player subsystem has initialized the user settings!");
 	}
 }
 
@@ -168,7 +168,7 @@ bool UEnhancedInputWorldSubsystem::ShouldCreateSubsystem(UObject* Outer) const
 	const bool bShouldCreate = GetDefault<UEnhancedInputDeveloperSettings>()->bEnableWorldSubsystem;
 	if (!bShouldCreate)
 	{
-		UE_LOG(LogWorldSubsystemInput, Verbose, TEXT("UEnhancedInputDeveloperSettings::bEnableWorldSubsystem is false, the world subsystem will not be created!"));
+		UE_LOGF(LogWorldSubsystemInput, Verbose, "UEnhancedInputDeveloperSettings::bEnableWorldSubsystem is false, the world subsystem will not be created!");
 	}
 
 	return bShouldCreate;
@@ -284,7 +284,7 @@ void UEnhancedInputWorldSubsystem::AddActorInputComponent(AActor* Actor)
 		}
 		else
 		{
-			UE_LOG(LogWorldSubsystemInput, Error, TEXT("Attempted to push a null Input Component to the Enhanced Input World Subsystem!"));
+			UE_LOGF(LogWorldSubsystemInput, Error, "Attempted to push a null Input Component to the Enhanced Input World Subsystem!");
 		}
 	}
 }
@@ -310,7 +310,7 @@ bool UEnhancedInputWorldSubsystem::InputKey(const FInputKeyEventArgs& Params)
 {
 	if (PlayerInput)
 	{
-		UE_LOG(LogWorldSubsystemInput, VeryVerbose, TEXT("EI %s World Subsystem InputKey : [%s]"), LexToString(GetWorld()->WorldType), *Params.Key.ToString());
+		UE_LOGF(LogWorldSubsystemInput, VeryVerbose, "EI %ls World Subsystem InputKey : [%ls]", LexToString(GetWorld()->WorldType), *Params.Key.ToString());
 		return PlayerInput->InputKey(Params);
 	}
 
diff --git a/Engine/Plugins/EnhancedInput/Source/EnhancedInput/Private/EnhancedPlayerInput.cpp b/Engine/Plugins/EnhancedInput/Source/EnhancedInput/Private/EnhancedPlayerInput.cpp
--- a/Engine/Plugins/EnhancedInput/Source/EnhancedInput/Private/EnhancedPlayerInput.cpp
+++ b/Engine/Plugins/EnhancedInput/Source/EnhancedInput/Private/EnhancedPlayerInput.cpp
@@ -782,7 +782,7 @@ bool UEnhancedPlayerInput::EvaluateInputComponentDelegates(UInputComponent* cons
 			// Skip input actions which have been marked for consumption by a higher priority UEnhancedInputComponent
 			if (UE::Input::bEnableListenerConsumption && ConsumedInputActions.Contains(DelegateAction))
 			{
-				UE_LOG(LogEnhancedInput, Verbose, TEXT("Input action '%s' was consumed by a higher priority input component listener. It will no longer fire"), *DelegateAction->GetName());
+				UE_LOGF(LogEnhancedInput, Verbose, "Input action '%ls' was consumed by a higher priority input component listener. It will no longer fire", *DelegateAction->GetName());
 				continue;
 			}
 
@@ -797,7 +797,7 @@ bool UEnhancedPlayerInput::EvaluateInputComponentDelegates(UInputComponent* cons
 						bCanTrigger &= !TriggeredActionsThisTick.Contains(DepAction.SourceAction);
 						if(!bCanTrigger)
 						{
-							UE_LOG(LogEnhancedInput, Verbose, TEXT("'%s' action was cancelled, its dependant on '%s'"), *DelegateAction->GetName(), *DepAction.SourceAction->GetName());
+							UE_LOGF(LogEnhancedInput, Verbose, "'%ls' action was cancelled, its dependant on '%ls'", *DelegateAction->GetName(), *DepAction.SourceAction->GetName());
 						}
 					}
 				}
@@ -961,7 +961,7 @@ FInputActionValue UEnhancedPlayerInput::GetActionValue(TObjectPtr<const UInputAc
 {
 	if (!ForAction)
 	{
-		UE_LOG(LogEnhancedInput, Error, TEXT("[%hs] Attempting to get the value of an invalid Input action!"), __func__);
+		UE_LOGF(LogEnhancedInput, Error, "[%hs] Attempting to get the value of an invalid Input action!", __func__);
 		return FInputActionValue();
 	}
 
@@ -1073,7 +1073,7 @@ void UEnhancedPlayerInput::NotifyInputActionsUnmapped(const TSet<const UInputAct
 			ActionInstanceData.Remove(Action);
 		}
 		
-		UE_LOG(LogEnhancedInput, VeryVerbose, TEXT("[%hs] EnhancedInput.ReconcileRemovedMappingDelegates is false! Using the legacy behavior to remove action instance data. Canceled events will not be fired."), __func__);
+		UE_LOGF(LogEnhancedInput, VeryVerbose, "[%hs] EnhancedInput.ReconcileRemovedMappingDelegates is false! Using the legacy behavior to remove action instance data. Canceled events will not be fired.", __func__);
 	}
 }
 
diff --git a/Engine/Plugins/EnhancedInput/Source/EnhancedInput/Private/InputTriggers.cpp b/Engine/Plugins/EnhancedInput/Source/EnhancedInput/Private/InputTriggers.cpp
--- a/Engine/Plugins/EnhancedInput/Source/EnhancedInput/Private/InputTriggers.cpp
+++ b/Engine/Plugins/EnhancedInput/Source/EnhancedInput/Private/InputTriggers.cpp
@@ -278,7 +278,7 @@ ETriggerState UInputTriggerRepeatedTap::UpdateState_Implementation(const UEnhanc
 		NumberOfTapsSinceLastTrigger = 0;
 	}
 	
-	UE_LOG(LogEnhancedInput, Verbose, TEXT("Repeated Tap InputState:  %d / %d taps :: %s"), NumberOfTapsSinceLastTrigger, NumberOfTapsWhichTriggerRepeat, *UE::Input::LexToString(State));
+	UE_LOGF(LogEnhancedInput, Verbose, "Repeated Tap InputState:  %d / %d taps :: %ls", NumberOfTapsSinceLastTrigger, NumberOfTapsWhichTriggerRepeat, *UE::Input::LexToString(State));
 	
 	return State;
 }
@@ -353,7 +353,7 @@ ETriggerState UInputTriggerCombo::UpdateState_Implementation(const UEnhancedPlay
 {
 	if (ComboActions.IsEmpty())
 	{
-		UE_LOG(LogEnhancedInput, Warning, TEXT("A Combo Trigger has no combo actions and will not work properly!"));
+		UE_LOGF(LogEnhancedInput, Warning, "A Combo Trigger has no combo actions and will not work properly!");
 		return ETriggerState::None;
 	}
 	
diff --git a/Engine/Plugins/EnhancedInput/Source/EnhancedInput/Private/UserSettings/EnhancedInputUserSettings.cpp b/Engine/Plugins/EnhancedInput/Source/EnhancedInput/Private/UserSettings/EnhancedInputUserSettings.cpp
--- a/Engine/Plugins/EnhancedInput/Source/EnhancedInput/Private/UserSettings/EnhancedInputUserSettings.cpp
+++ b/Engine/Plugins/EnhancedInput/Source/EnhancedInput/Private/UserSettings/EnhancedInputUserSettings.cpp
@@ -74,7 +74,7 @@ namespace UE::EnhancedInput
 	{
 		if (!Args.IsValidIndex(0))
 		{
-			UE_LOG(LogEnhancedInput, Error, TEXT("EnhancedInput.SetActiveKeyProfile command failed - you must provide a desired key profile name! Example usage:\n\t EnhancedInput.SetActiveKeyProfile PresetA"));
+			UE_LOGF(LogEnhancedInput, Error, "EnhancedInput.SetActiveKeyProfile command failed - you must provide a desired key profile name! Example usage:\n\t EnhancedInput.SetActiveKeyProfile PresetA");
 			return;
 		}
 
@@ -373,17 +373,17 @@ void UEnhancedPlayerMappableKeyProfile::ResetToDefault()
 		}
 	}
 	
-	UE_LOG(LogEnhancedInput, Verbose, TEXT("Reset Player Mappable Key Profile '%s' to default values"), *ProfileIdentifierString);
+	UE_LOGF(LogEnhancedInput, Verbose, "Reset Player Mappable Key Profile '%ls' to default values", *ProfileIdentifierString);
 }
 
 void UEnhancedPlayerMappableKeyProfile::EquipProfile()
 {
-	UE_LOG(LogEnhancedInput, Verbose, TEXT("Equipping Player Mappable Key Profile '%s'"), *ProfileIdentifierString);
+	UE_LOGF(LogEnhancedInput, Verbose, "Equipping Player Mappable Key Profile '%ls'", *ProfileIdentifierString);
 }
 
 void UEnhancedPlayerMappableKeyProfile::UnEquipProfile()
 {
-	UE_LOG(LogEnhancedInput, Verbose, TEXT("Unequipping Player Mappable Key Profile '%s'"), *ProfileIdentifierString);
+	UE_LOGF(LogEnhancedInput, Verbose, "Unequipping Player Mappable Key Profile '%ls'", *ProfileIdentifierString);
 }
 
 void UEnhancedPlayerMappableKeyProfile::SetDisplayName(const FText& NewDisplayName)
@@ -436,7 +436,7 @@ const FKeyMappingRow* UEnhancedPlayerMappableKeyProfile::FindKeyMappingRow(const
 
 void UEnhancedPlayerMappableKeyProfile::DumpProfileToLog() const
 {
-	UE_LOG(LogEnhancedInput, Log, TEXT("%s"), *ToString());
+	UE_LOGF(LogEnhancedInput, Log, "%ls", *ToString());
 }
 
 FString UEnhancedPlayerMappableKeyProfile::ToString() const
@@ -471,7 +471,7 @@ int32 UEnhancedPlayerMappableKeyProfile::GetMappedKeysInRow(const FName MappingN
 	}
 	else
 	{
-		UE_LOG(LogEnhancedInput, Verbose, TEXT("Player Mappable Key Profile '%s' doesn't have any mappings for '%s'"), *ProfileIdentifierString, *MappingName.ToString());
+		UE_LOGF(LogEnhancedInput, Verbose, "Player Mappable Key Profile '%ls' doesn't have any mappings for '%ls'", *ProfileIdentifierString, *MappingName.ToString());
 	}
 	return OutKeys.Num();
 }
@@ -519,7 +519,7 @@ int32 UEnhancedPlayerMappableKeyProfile::QueryPlayerMappedKeys(const FPlayerMapp
 	}
 	else
 	{
-		UE_LOG(LogEnhancedInput, Verbose, TEXT("Player Mappable Key Profile '%s' doesn't have any mappings for '%s'"), *ProfileIdentifierString, *Options.MappingName.ToString());
+		UE_LOGF(LogEnhancedInput, Verbose, "Player Mappable Key Profile '%ls' doesn't have any mappings for '%ls'", *ProfileIdentifierString, *Options.MappingName.ToString());
 	}
 
 	return OutKeys.Num();
@@ -657,7 +657,7 @@ UEnhancedInputUserSettings* UEnhancedInputUserSettings::LoadOrCreateSettings(ULo
 
 	if (!LocalPlayer)
 	{
-		UE_LOG(LogEnhancedInput, Log, TEXT("Unable to determine an owning Local Player for the given Enhanced Player Input object"));
+		UE_LOGF(LogEnhancedInput, Log, "Unable to determine an owning Local Player for the given Enhanced Player Input object");
 		return nullptr;
 	}
 
@@ -691,11 +691,11 @@ UEnhancedInputUserSettings* UEnhancedInputUserSettings::LoadOrCreateSettings(ULo
 	// or if you change the class of your user settings object and would to keep some of the old data.
 	else if (Settings->GetClass() != SettingsClass)
 	{
-		UE_LOG(
+		UE_LOGF(
 			LogEnhancedInput,
 			Log,
-			TEXT("[%hs] Enhanced Input User Settings class has been changed from '%s' to '%s'. A new settings save object will be created (your saved settings will be reset). Override HandleSettingsClassChangedOnLoad_Impl to retain data."),
-			__func__, *GetNameSafe(Settings->GetClass()), *GetNameSafe(SettingsClass))
+			"[%hs] Enhanced Input User Settings class has been changed from '%ls' to '%ls'. A new settings save object will be created (your saved settings will be reset). Override HandleSettingsClassChangedOnLoad_Impl to retain data.",
+			__func__, *GetNameSafe(Settings->GetClass()), *GetNameSafe(SettingsClass));
 
 		// Create a new settings object here with the new class which it has been changed to
 		UEnhancedInputUserSettings* NewSettingsObject = Cast<UEnhancedInputUserSettings>(UGameplayStatics::CreateSaveGameObject(SettingsClass));
@@ -748,7 +748,7 @@ void UEnhancedInputUserSettings::Initialize(ULocalPlayer* InLocalPlayer)
 
 void UEnhancedInputUserSettings::ApplySettings()
 {
-	UE_LOG(LogEnhancedInput, Verbose, TEXT("Enhanced Input User Settings applied!"));
+	UE_LOGF(LogEnhancedInput, Verbose, "Enhanced Input User Settings applied!");
 
 	OnSettingsApplied.Broadcast();
 }
@@ -758,11 +758,11 @@ void UEnhancedInputUserSettings::SaveSettings()
 	if (ULocalPlayer* OwningPlayer = GetLocalPlayer())
 	{
 		UGameplayStatics::SaveGameToSlot(this, GetSaveFilename(OwningPlayer), OwningPlayer->GetLocalPlayerIndex());
-		UE_LOG(LogEnhancedInput, Verbose, TEXT("Enhanced Input User Settings saved!"));
+		UE_LOGF(LogEnhancedInput, Verbose, "Enhanced Input User Settings saved!");
 	}
 	else
 	{
-		UE_LOG(LogEnhancedInput, Warning, TEXT("Attempting to save Enhanced Input User settings without an owning local player!"));
+		UE_LOGF(LogEnhancedInput, Warning, "Attempting to save Enhanced Input User settings without an owning local player!");
 	}
 }
 
@@ -777,7 +777,7 @@ void UEnhancedInputUserSettings::AsyncSaveSettings()
 	}
 	else
 	{
-		UE_LOG(LogEnhancedInput, Error, TEXT("[UEnhancedInputUserSettings::AsyncSaveSettings] Failed async save, there is no owning local player!"));
+		UE_LOGF(LogEnhancedInput, Error, "[UEnhancedInputUserSettings::AsyncSaveSettings] Failed async save, there is no owning local player!");
 	}
 }
 
@@ -825,11 +825,11 @@ void UEnhancedInputUserSettings::OnAsyncSaveComplete(const FString& SlotName, co
 {
 	if (bSuccess)
 	{
-		UE_LOG(LogEnhancedInput, Log, TEXT("UEnhancedInputUserSettings::OnAsyncSaveComplete] Async save of slot '%s' for user index '%d' completed successfully!"), *SlotName, UserIndex);
+		UE_LOGF(LogEnhancedInput, Log, "UEnhancedInputUserSettings::OnAsyncSaveComplete] Async save of slot '%ls' for user index '%d' completed successfully!", *SlotName, UserIndex);
 	}
 	else
 	{
-		UE_LOG(LogEnhancedInput, Error, TEXT("[UEnhancedInputUserSettings::OnAsyncSaveComplete] Failed async save of slot '%s' for user index '%d'"), *SlotName, UserIndex);
+		UE_LOGF(LogEnhancedInput, Error, "[UEnhancedInputUserSettings::OnAsyncSaveComplete] Failed async save of slot '%ls' for user index '%d'", *SlotName, UserIndex);
 	}
 }
 
@@ -1106,7 +1106,7 @@ void UEnhancedInputUserSettings::Serialize(FArchive& Ar)
 #else
 		if (!ensure(SavedObjectTerminator == TEXT("ObjectEnd")))
 		{
-			UE_LOG(LogEnhancedInput, Error, TEXT("Serialization size mismatch! Possible over-read or over-write of this buffer."));
+			UE_LOGF(LogEnhancedInput, Error, "Serialization size mismatch! Possible over-read or over-write of this buffer.");
 			break;
 		}
 #endif	// WITH_EDITOR
@@ -1204,7 +1204,7 @@ void UEnhancedInputUserSettings::MapPlayerKey(const FMapPlayerKeyArgs& InArgs, F
 					}
 					else
 					{
-						UE_LOG(LogEnhancedInput, Log, TEXT("[UEnhancedInputUserSettings::MapPlayerKey] Unable to find a matching Hardware Device Identifier with the HardwareDeviceId of '%s'"), *InArgs.HardwareDeviceId.ToString());
+						UE_LOGF(LogEnhancedInput, Log, "[UEnhancedInputUserSettings::MapPlayerKey] Unable to find a matching Hardware Device Identifier with the HardwareDeviceId of '%ls'", *InArgs.HardwareDeviceId.ToString());
 					}
 				}
 			}
@@ -1245,7 +1245,7 @@ void UEnhancedInputUserSettings::MapPlayerKey(const FMapPlayerKeyArgs& InArgs, F
 	{
 		// If there is neither an existing key mapping or a key mapping row, then mapping the key has failed
 		FailureReason.AddTag(UE::EnhancedInput::TAG_NoMappingRowFound);
-		UE_LOG(LogEnhancedInput, Warning, TEXT("[UEnhancedInputUserSettings::MapPlayerKey] Failed to map a player key for '%s'"), *InArgs.MappingName.ToString());
+		UE_LOGF(LogEnhancedInput, Warning, "[UEnhancedInputUserSettings::MapPlayerKey] Failed to map a player key for '%ls'", *InArgs.MappingName.ToString());
 	}
 }
 
@@ -1297,7 +1297,7 @@ void UEnhancedInputUserSettings::UnMapPlayerKey(const FMapPlayerKeyArgs& InArgs,
 			}
 		}
 		
-		UE_LOG(LogEnhancedInput, Verbose, TEXT("[UEnhancedInputUserSettings::MapPlayerKey] Reset keymapping to default: '%s'"), *FoundMapping->ToString());
+		UE_LOGF(LogEnhancedInput, Verbose, "[UEnhancedInputUserSettings::MapPlayerKey] Reset keymapping to default: '%ls'", *FoundMapping->ToString());
 	}
 	// if a mapping doesn't exist, then we can't unmap it
 	else
@@ -1335,7 +1335,7 @@ void UEnhancedInputUserSettings::ResetAllPlayerKeysInRow(const FMapPlayerKeyArgs
 				OnKeyMappingUpdated(&Mapping, InArgs, true);
 				OnSettingsChanged.Broadcast(this);
 			}
-			UE_LOG(LogEnhancedInput, Verbose, TEXT("[UEnhancedInputUserSettings::MapPlayerKey] Reset keymapping to default: '%s'"), *Mapping.ToString());
+			UE_LOGF(LogEnhancedInput, Verbose, "[UEnhancedInputUserSettings::MapPlayerKey] Reset keymapping to default: '%ls'", *Mapping.ToString());
 		}
 	}
 	else
@@ -1377,7 +1377,7 @@ const TSet<FPlayerKeyMapping>& UEnhancedInputUserSettings::FindMappingsInRow(con
 		return ExistingMappings.Mappings;
 	}
 	
-	UE_LOG(LogEnhancedInput, Error, TEXT("There is no current mappable key profile! No mappings will be returned."));
+	UE_LOGF(LogEnhancedInput, Error, "There is no current mappable key profile! No mappings will be returned.");
 	
 	static TSet<FPlayerKeyMapping> EmptyMappings;
 	return EmptyMappings;
@@ -1394,7 +1394,7 @@ const FPlayerKeyMapping* UEnhancedInputUserSettings::FindCurrentMappingForSlot(c
 		}
 	}
 	
-	UE_LOG(LogEnhancedInput, Warning, TEXT("No mappings could be found for action '%s'"), *ActionName.ToString());
+	UE_LOGF(LogEnhancedInput, Warning, "No mappings could be found for action '%ls'", *ActionName.ToString());
 	return nullptr;
 }
 
@@ -1428,7 +1428,7 @@ bool UEnhancedInputUserSettings::SetActiveKeyProfile(const FString& InProfileId)
 	// If this key profile is already equipped, then there is nothing to do
 	if (InProfileId.Equals(OriginalProfileId))
 	{
-		UE_LOG(LogEnhancedInput, Log, TEXT("Key profile '%s' is already currently equipped. Nothing will happen."), *InProfileId);
+		UE_LOGF(LogEnhancedInput, Log, "Key profile '%ls' is already currently equipped. Nothing will happen.", *InProfileId);
 		return false;
 	}
 	
@@ -1454,11 +1454,11 @@ bool UEnhancedInputUserSettings::SetActiveKeyProfile(const FString& InProfileId)
 	}
 	else
 	{
-		UE_LOG(LogEnhancedInput, Error, TEXT("No profile with name '%s' exists! Did you call CreateNewKeyProfile at any point?"), *InProfileId);
+		UE_LOGF(LogEnhancedInput, Error, "No profile with name '%ls' exists! Did you call CreateNewKeyProfile at any point?", *InProfileId);
 		return false;
 	}
 
-	UE_LOG(LogEnhancedInput, Verbose, TEXT("Successfully changed Key Profile from '%s' to '%s'"), *OriginalProfileId, *CurrentProfileIdentifierString);
+	UE_LOGF(LogEnhancedInput, Verbose, "Successfully changed Key Profile from '%ls' to '%ls'", *OriginalProfileId, *CurrentProfileIdentifierString);
 	return true;
 }
 
@@ -1528,7 +1528,7 @@ UEnhancedPlayerMappableKeyProfile* UEnhancedInputUserSettings::CreateNewKeyProfi
 {
 	if (!InArgs.ProfileType)
 	{
-		UE_LOG(LogEnhancedInput, Error, TEXT("Invalid ProfileType given CreateNewKeyProfile!"));
+		UE_LOGF(LogEnhancedInput, Error, "Invalid ProfileType given CreateNewKeyProfile!");
 		return nullptr;
 	}
 	
@@ -1537,7 +1537,7 @@ UEnhancedPlayerMappableKeyProfile* UEnhancedInputUserSettings::CreateNewKeyProfi
 	// Check for an existing profile of this name
 	if (OutProfile)
 	{
-		UE_LOG(LogEnhancedInput, Warning, TEXT("A key profile with the name '%s' already exists! Use a different name."), *InArgs.ProfileStringIdentifier);
+		UE_LOGF(LogEnhancedInput, Warning, "A key profile with the name '%ls' already exists! Use a different name.", *InArgs.ProfileStringIdentifier);
 	}
 	else
 	{
@@ -1562,7 +1562,7 @@ UEnhancedPlayerMappableKeyProfile* UEnhancedInputUserSettings::CreateNewKeyProfi
 		SetActiveKeyProfile(InArgs.ProfileStringIdentifier);
 	}
 
-	UE_LOG(LogEnhancedInput, Verbose, TEXT("Completed creation of key mapping profile '%s'"), *OutProfile->ProfileIdentifierString);
+	UE_LOGF(LogEnhancedInput, Verbose, "Completed creation of key mapping profile '%ls'", *OutProfile->ProfileIdentifierString);
 	
 	return OutProfile;
 }
@@ -1599,7 +1599,7 @@ bool UEnhancedInputUserSettings::RegisterInputMappingContext(const UInputMapping
 {
 	if (!IMC)
 	{
-		UE_LOG(LogEnhancedInput, Error, TEXT("Attempting to register a null mapping context with the user settings!"));
+		UE_LOGF(LogEnhancedInput, Error, "Attempting to register a null mapping context with the user settings!");
 		ensure(false);
 		return false;
 	}
@@ -1620,7 +1620,7 @@ bool UEnhancedInputUserSettings::RegisterInputMappingContextInternal(const UInpu
 {
 	if (!IMC)
 	{
-		UE_LOG(LogEnhancedInput, Error, TEXT("Attempting to register a null mapping context with the user settings!"));
+		UE_LOGF(LogEnhancedInput, Error, "Attempting to register a null mapping context with the user settings!");
 		ensure(false);
 		return false;
 	}	
@@ -1665,7 +1665,7 @@ bool UEnhancedInputUserSettings::RegisterKeyMappingsToProfile(UEnhancedPlayerMap
 {
 	if (!IMC)
 	{
-		UE_LOG(LogEnhancedInput, Error, TEXT("Attempting to register a null mapping context with the user settings!"));
+		UE_LOGF(LogEnhancedInput, Error, "Attempting to register a null mapping context with the user settings!");
 		ensure(false);
 		return false;
 	}
@@ -1762,7 +1762,7 @@ bool UEnhancedInputUserSettings::RegisterKeyMappingsToProfile(UEnhancedPlayerMap
 		MappingSlot = static_cast<EPlayerMappableKeySlot>(FMath::Min<uint8>(static_cast<uint8>(MappingSlot) + 1, static_cast<uint8>(EPlayerMappableKeySlot::Max)));
 	}
 
-	UE_LOG(LogEnhancedInput, Verbose, TEXT("Registered IMC with UEnhancedInputUserSettings: %s"), *IMC->GetFName().ToString());
+	UE_LOGF(LogEnhancedInput, Verbose, "Registered IMC with UEnhancedInputUserSettings: %ls", *IMC->GetFName().ToString());
 	return true;
 }
 
diff --git a/Engine/Plugins/EnhancedInput/Source/InputEditor/Private/EnhancedInputEditorSubsystem.cpp b/Engine/Plugins/EnhancedInput/Source/InputEditor/Private/EnhancedInputEditorSubsystem.cpp
--- a/Engine/Plugins/EnhancedInput/Source/InputEditor/Private/EnhancedInputEditorSubsystem.cpp
+++ b/Engine/Plugins/EnhancedInput/Source/InputEditor/Private/EnhancedInputEditorSubsystem.cpp
@@ -99,7 +99,7 @@ void UEnhancedInputEditorSubsystem::Tick(float DeltaTime)
 {
 	if (!ensure(PlayerInput))
 	{
-		UE_LOG(LogEditorInput, Error, TEXT("UEnhancedInputEditorSubsystem is ticking before it has created a PlayerInput!"));
+		UE_LOGF(LogEditorInput, Error, "UEnhancedInputEditorSubsystem is ticking before it has created a PlayerInput!");
 		return;
 	}
 
@@ -140,7 +140,7 @@ void UEnhancedInputEditorSubsystem::PushInputComponent(UInputComponent* InInputC
 	{
 		if (GetDefault<UEnhancedInputEditorSettings>()->bLogAllInput)
 		{
-			UE_LOG(LogEditorInput, Log, TEXT("PUSHING AN INPUT COMPONENT! '%s'"), *InInputComponent->GetName());
+			UE_LOGF(LogEditorInput, Log, "PUSHING AN INPUT COMPONENT! '%ls'", *InInputComponent->GetName());
 		}
 
 		// Mark this input component for firing delegates in the editor
@@ -173,7 +173,7 @@ void UEnhancedInputEditorSubsystem::PushInputComponent(UInputComponent* InInputC
 	}
 	else
 	{
-		UE_LOG(LogEditorInput, Error, TEXT("Attempted to push a null Input Component to the Enhanced Input Editor Subsystem!"));
+		UE_LOGF(LogEditorInput, Error, "Attempted to push a null Input Component to the Enhanced Input Editor Subsystem!");
 	}
 }
 
@@ -239,7 +239,7 @@ bool UEnhancedInputEditorSubsystem::InputKey(const FInputKeyEventArgs& Params)
 		{
 			if (GetDefault<UEnhancedInputEditorSettings>()->bLogAllInput)
 			{
-				UE_LOG(LogEditorInput, Log, TEXT("EI Editor Subsystem InputKey : [%s]"), *Params.Key.ToString());
+				UE_LOGF(LogEditorInput, Log, "EI Editor Subsystem InputKey : [%ls]", *Params.Key.ToString());
 			}
 			return PlayerInput->InputKey(Params);
 		}
diff --git a/Engine/Plugins/EnhancedInput/Source/InputEditor/Private/InputEditorModule.cpp b/Engine/Plugins/EnhancedInput/Source/InputEditor/Private/InputEditorModule.cpp
--- a/Engine/Plugins/EnhancedInput/Source/InputEditor/Private/InputEditorModule.cpp
+++ b/Engine/Plugins/EnhancedInput/Source/InputEditor/Private/InputEditorModule.cpp
@@ -419,7 +419,7 @@ namespace UE::Input
 		Info->bFireAndForget = true;
 		FSlateNotificationManager::Get().QueueNotification(Info);
 		
-		UE_LOG(LogEnhancedInputEditor, Log, TEXT("Upgrading Default %s class from '%s' to '%s'"), *ClassBeingChanged.ToString(), *OldClassName, *NewClassName);
+		UE_LOGF(LogEnhancedInputEditor, Log, "Upgrading Default %ls class from '%ls' to '%ls'", *ClassBeingChanged.ToString(), *OldClassName, *NewClassName);
 	}
 }
 
diff --git a/Engine/Plugins/EnhancedInput/Source/InputEditor/Private/Tests/InputModifierTest.cpp b/Engine/Plugins/EnhancedInput/Source/InputEditor/Private/Tests/InputModifierTest.cpp
--- a/Engine/Plugins/EnhancedInput/Source/InputEditor/Private/Tests/InputModifierTest.cpp
+++ b/Engine/Plugins/EnhancedInput/Source/InputEditor/Private/Tests/InputModifierTest.cpp
@@ -308,7 +308,7 @@ bool FInputModifierPerformanceTest::RunTest(const FString& Parameters)
 	check(AverageTimings.Num() % 2 == 0);
 	const double MedianTiming = (AverageTimings[NumTimesToRun / 2 - 1] + AverageTimings[NumTimesToRun / 2]) / 2;;
 
-	UE_LOG(LogEnhancedInput, Log, TEXT("Modifiers Perf Test (in seconds)... Median: %lf   Min: %lf   Max: %lf"), MedianTiming, MinRun, MaxRun);
+	UE_LOGF(LogEnhancedInput, Log, "Modifiers Perf Test (in seconds)... Median: %lf   Min: %lf   Max: %lf", MedianTiming, MinRun, MaxRun);
 	
 	TestLessEqual(TEXT("MinRun is less then max"), MinRun, MaxRun);
 	TestNotEqual(TEXT("Median Time was calculated"), MedianTiming, 0.0);
