[Enhanced Input] Add a function to reset the current key profile to default. This is just a useful helper for users of this API instead of needing to reference different profile ID's.

#rnx

#Synced-CL 51092797

diff --git a/Engine/Plugins/EnhancedInput/Source/EnhancedInput/Private/UserSettings/EnhancedInputUserSettings.cpp b/Engine/Plugins/EnhancedInput/Source/EnhancedInput/Private/UserSettings/EnhancedInputUserSettings.cpp
--- a/Engine/Plugins/EnhancedInput/Source/EnhancedInput/Private/UserSettings/EnhancedInputUserSettings.cpp
+++ b/Engine/Plugins/EnhancedInput/Source/EnhancedInput/Private/UserSettings/EnhancedInputUserSettings.cpp
@@ -1462,6 +1462,18 @@ bool UEnhancedInputUserSettings::SetActiveKeyProfile(const FString& InProfileId)
 	return true;
 }
 
+bool UEnhancedInputUserSettings::SetKeyProfileToDefault()
+{
+	if (const UEnhancedPlayerMappableKeyProfile* DefaultProfile = GetDefaultKeyProfile())
+	{
+		return SetActiveKeyProfile(DefaultProfile->GetProfileIdString());	
+	}
+	else
+	{
+		return false;
+	}
+}
+
 PRAGMA_DISABLE_DEPRECATION_WARNINGS
 const FGameplayTag& UEnhancedInputUserSettings::GetCurrentKeyProfileIdentifier() const
 {
diff --git a/Engine/Plugins/EnhancedInput/Source/EnhancedInput/Public/UserSettings/EnhancedInputUserSettings.h b/Engine/Plugins/EnhancedInput/Source/EnhancedInput/Public/UserSettings/EnhancedInputUserSettings.h
--- a/Engine/Plugins/EnhancedInput/Source/EnhancedInput/Public/UserSettings/EnhancedInputUserSettings.h
+++ b/Engine/Plugins/EnhancedInput/Source/EnhancedInput/Public/UserSettings/EnhancedInputUserSettings.h
@@ -711,6 +711,12 @@ public:
 	 */
 	UFUNCTION(BlueprintCallable, Category="Enhanced Input|User Settings", meta=(ReturnDisplayName = "Was Successful"))
 	UE_API virtual bool SetActiveKeyProfile(const FString& InProfileId);
+
+	/**
+	 * Changes the currently active key profile the default profile. Returns true if the operation was successful.
+	 */
+	UFUNCTION(BlueprintCallable, Category="Enhanced Input|User Settings", meta=(ReturnDisplayName = "Was Successful"))
+	UE_API bool SetKeyProfileToDefault();
 	
 	/** Gets the currently selected key profile */
 	UE_DEPRECATED(5.6, "Use GetActiveKeyProfileId instead ")
