[Enhanced Input] Add a const array view to the EnhancedActionMappings. This is useful for users if they wish to query the current state of a key mapping in Enhanced Input, like in UI code if you want to know the total time of a "hold" trigger.

#rnx
#Synced-CL 51542463

diff --git a/Engine/Plugins/EnhancedInput/Source/EnhancedInput/Private/EnhancedPlayerInput.cpp b/Engine/Plugins/EnhancedInput/Source/EnhancedInput/Private/EnhancedPlayerInput.cpp
--- a/Engine/Plugins/EnhancedInput/Source/EnhancedInput/Private/EnhancedPlayerInput.cpp
+++ b/Engine/Plugins/EnhancedInput/Source/EnhancedInput/Private/EnhancedPlayerInput.cpp
@@ -336,6 +336,11 @@ FString UEnhancedPlayerInput::GetUserSettingsSaveFileName() const
 	return GetDefault<UEnhancedInputDeveloperSettings>()->InputSettingsSaveSlotName;
 }
 
+TConstArrayView<const FEnhancedActionKeyMapping> UEnhancedPlayerInput::GetEnhancedActionMappingsView() const
+{
+	return { EnhancedActionMappings };
+}
+
 void UEnhancedPlayerInput::SetCurrentInputMode(const FGameplayTagContainer& NewMode)
 {
 	CurrentInputMode = NewMode;
diff --git a/Engine/Plugins/EnhancedInput/Source/EnhancedInput/Public/EnhancedPlayerInput.h b/Engine/Plugins/EnhancedInput/Source/EnhancedInput/Public/EnhancedPlayerInput.h
--- a/Engine/Plugins/EnhancedInput/Source/EnhancedInput/Public/EnhancedPlayerInput.h
+++ b/Engine/Plugins/EnhancedInput/Source/EnhancedInput/Public/EnhancedPlayerInput.h
@@ -135,6 +135,15 @@ public:
 	 * By default, this will return InputSettingsSaveSlotName from the developer settings.
 	 */
 	UE_API virtual FString GetUserSettingsSaveFileName() const;
+
+	/**
+	 * Returns a const array view of the current enhanced input action mappings.
+	 * 
+	 * Note: Do NOT hold onto any references to these mappings or their input triggers/modifiers, as they will get
+	 * re-instanced every time the control mappings are rebuilt. This is for read-only access to view the current mappings
+	 * which may be useful for systems which want to query specific data about the current trigger states of specific mappings.
+	 */
+	UE_API TConstArrayView<const FEnhancedActionKeyMapping> GetEnhancedActionMappingsView() const;
 	
 protected:
 
