[Enhanced Input] Add some more automated tests to Enhanced Input for modifiers, triggers, and more.

#rnx
#Synced-CL 52004614

diff --git a/Engine/Plugins/EnhancedInput/Source/InputEditor/Private/Tests/InputModifierExtendedTest.cpp b/Engine/Plugins/EnhancedInput/Source/InputEditor/Private/Tests/InputModifierExtendedTest.cpp
new file mode 100644
--- /dev/null
+++ b/Engine/Plugins/EnhancedInput/Source/InputEditor/Private/Tests/InputModifierExtendedTest.cpp
@@ -0,0 +1,447 @@
+// Copyright Epic Games, Inc. All Rights Reserved.
+
+#include "InputTestFramework.h"
+#include "Misc/AutomationTest.h"
+#include "InputModifiers.h"
+#include "EnhancedActionKeyMapping.h"
+
+// Extended modifier tests covering SwizzleAxis, per-axis Negate, Radial DeadZone,
+// ResponseCurveExponential, ScaleByDeltaTime, and modifier application order.
+
+constexpr auto BasicModifierExtTestFlags = EAutomationTestFlags::EditorContext | EAutomationTestFlags::EngineFilter;
+
+// Helper: set up a world+player+context+action, return a Data reference.
+static UControllablePlayer& ABasicModifierExtTest(FAutomationTestBase* Test, EInputActionValueType ForValueType)
+{
+	UWorld* World =
+	GIVEN(AnEmptyWorld());
+
+	UControllablePlayer& Data =
+	AND(AControllablePlayer(World));
+	Test->TestTrue(TEXT("Controllable Player is valid"), Data.IsValid());
+	AND(AnInputContextIsAppliedToAPlayer(Data, TestContext, 0));
+	AND(AnInputAction(Data, TestAction, ForValueType));
+
+	return Data;
+}
+
+// Helper: read the float value from action instance data.
+static float GetModExtActionValue(UControllablePlayer& Data, FName ActionName)
+{
+	return FInputTestHelper::GetActionData(Data, ActionName).GetValue().Get<float>();
+}
+
+// FVector::X/Y/Z are double in UE5. These helpers avoid ambiguous TestEqual overload resolution.
+static float FX(const FVector& V) { return (float)V.X; }
+static float FY(const FVector& V) { return (float)V.Y; }
+static float FZ(const FVector& V) { return (float)V.Z; }
+
+// ============================================================
+// Swizzle Axis
+// ============================================================
+
+IMPLEMENT_SIMPLE_AUTOMATION_TEST(FInputModifierSwizzleAxisTest, "Input.Modifiers.SwizzleAxis", BasicModifierExtTestFlags)
+
+bool FInputModifierSwizzleAxisTest::RunTest(const FString& Parameters)
+{
+	// Test 1 - YXZ: swap X and Y.  Most common use: drive the Y component of a 2D action from a 1D analog input.
+	{
+		UControllablePlayer& Data =
+		GIVEN(ABasicModifierExtTest(this, EInputActionValueType::Axis2D));
+		AND(AnActionIsMappedToAKey(Data, TestContext, TestAction, EKeys::Gamepad_Left2D));
+
+		UInputModifierSwizzleAxis* Swizzle =
+		AND(Cast<UInputModifierSwizzleAxis>(AModifierIsAppliedToAnActionMapping(Data, NewObject<UInputModifierSwizzleAxis>(), TestContext, TestAction, EKeys::Gamepad_Left2D)));
+		check(Swizzle);
+		Swizzle->Order = EInputAxisSwizzle::YXZ;
+
+		// Drive LeftX=0.6, LeftY=0.8 -> raw FVector(0.6, 0.8, 0) -> after YXZ swap -> FVector(0.8, 0.6, 0)
+		WHEN(AKeyIsActuated(Data, EKeys::Gamepad_LeftX, 0.6f));
+		AND(AKeyIsActuated(Data, EKeys::Gamepad_LeftY, 0.8f));
+		AND(InputIsTicked(Data));
+		THEN(PressingKeyTriggersAction(Data, TestAction));
+
+		FVector Result = FInputTestHelper::GetTriggered<FVector>(Data, TestAction);
+		AND(TestEqual(TEXT("YXZ swap: X becomes Y"), FX(Result), 0.8f));
+		AND(TestEqual(TEXT("YXZ swap: Y becomes X"), FY(Result), 0.6f));
+	}
+
+	// Test 2 - ZYX: swap X and Z.  Useful when a 1D key should drive the Z component of a 3D action.
+	// A 1D key produces FVector(val, 0, 0).  After ZYX swap: FVector(0, 0, val).
+	{
+		UControllablePlayer& Data =
+		GIVEN(ABasicModifierExtTest(this, EInputActionValueType::Axis3D));
+		AND(AnActionIsMappedToAKey(Data, TestContext, TestAction, TestAxis));
+
+		UInputModifierSwizzleAxis* Swizzle =
+		AND(Cast<UInputModifierSwizzleAxis>(AModifierIsAppliedToAnActionMapping(Data, NewObject<UInputModifierSwizzleAxis>(), TestContext, TestAction, TestAxis)));
+		check(Swizzle);
+		Swizzle->Order = EInputAxisSwizzle::ZYX;
+
+		// TestAxis is a 1D analog -> FVector(0.7, 0, 0) -> ZYX swap X and Z -> FVector(0, 0, 0.7)
+		WHEN(AKeyIsActuated(Data, TestAxis, 0.7f));
+		AND(InputIsTicked(Data));
+		THEN(PressingKeyTriggersAction(Data, TestAction));
+
+		FVector Result = FInputTestHelper::GetTriggered<FVector>(Data, TestAction);
+		AND(TestEqual(TEXT("ZYX: X moved to Z"), FX(Result), 0.f));
+		AND(TestEqual(TEXT("ZYX: Y unchanged (zero)"), FY(Result), 0.f));
+		AND(TestEqual(TEXT("ZYX: Z receives X value"), FZ(Result), 0.7f));
+	}
+
+	// Test 3 - XZY: swap Y and Z.  A 1D key (X-only) should be unaffected, confirming correct axis targeting.
+	// FVector(0.5, 0, 0) -> XZY swap Y and Z -> FVector(0.5, 0, 0) (no change when Y=Z=0)
+	{
+		UControllablePlayer& Data =
+		GIVEN(ABasicModifierExtTest(this, EInputActionValueType::Axis3D));
+		AND(AnActionIsMappedToAKey(Data, TestContext, TestAction, TestAxis));
+
+		UInputModifierSwizzleAxis* Swizzle =
+		AND(Cast<UInputModifierSwizzleAxis>(AModifierIsAppliedToAnActionMapping(Data, NewObject<UInputModifierSwizzleAxis>(), TestContext, TestAction, TestAxis)));
+		check(Swizzle);
+		Swizzle->Order = EInputAxisSwizzle::XZY;
+
+		WHEN(AKeyIsActuated(Data, TestAxis, 0.5f));
+		AND(InputIsTicked(Data));
+		THEN(PressingKeyTriggersAction(Data, TestAction));
+
+		FVector Result = FInputTestHelper::GetTriggered<FVector>(Data, TestAction);
+		AND(TestEqual(TEXT("XZY: X unchanged"), FX(Result), 0.5f));
+		AND(TestEqual(TEXT("XZY: Y and Z both zero (swapped zeros)"), FY(Result), 0.f));
+		AND(TestEqual(TEXT("XZY: Z unchanged (swapped from zero Y)"), FZ(Result), 0.f));
+	}
+
+	return true;
+}
+
+// ============================================================
+// Negate per-axis flags
+// ============================================================
+
+IMPLEMENT_SIMPLE_AUTOMATION_TEST(FInputModifierNegatePerAxisTest, "Input.Modifiers.Negate.PerAxis", BasicModifierExtTestFlags)
+
+bool FInputModifierNegatePerAxisTest::RunTest(const FString& Parameters)
+{
+	// Use Axis2D so we have independent X and Y values to negate selectively.
+	// Drive LeftX=0.5, LeftY=0.5 -> raw FVector(0.5, 0.5, 0).
+
+	// Test 1 - negate X only
+	{
+		UControllablePlayer& Data =
+		GIVEN(ABasicModifierExtTest(this, EInputActionValueType::Axis2D));
+		AND(AnActionIsMappedToAKey(Data, TestContext, TestAction, EKeys::Gamepad_Left2D));
+
+		UInputModifierNegate* Negate =
+		AND(Cast<UInputModifierNegate>(AModifierIsAppliedToAnActionMapping(Data, NewObject<UInputModifierNegate>(), TestContext, TestAction, EKeys::Gamepad_Left2D)));
+		check(Negate);
+		Negate->bX = true;
+		Negate->bY = false;
+		Negate->bZ = false;
+
+		WHEN(AKeyIsActuated(Data, EKeys::Gamepad_LeftX, 0.5f));
+		AND(AKeyIsActuated(Data, EKeys::Gamepad_LeftY, 0.5f));
+		AND(InputIsTicked(Data));
+		THEN(PressingKeyTriggersAction(Data, TestAction));
+
+		FVector Result = FInputTestHelper::GetTriggered<FVector>(Data, TestAction);
+		AND(TestEqual(TEXT("Negate X only: X negated"), FX(Result), -0.5f));
+		AND(TestEqual(TEXT("Negate X only: Y unchanged"), FY(Result), 0.5f));
+	}
+
+	// Test 2 - negate Y only
+	{
+		UControllablePlayer& Data =
+		GIVEN(ABasicModifierExtTest(this, EInputActionValueType::Axis2D));
+		AND(AnActionIsMappedToAKey(Data, TestContext, TestAction, EKeys::Gamepad_Left2D));
+
+		UInputModifierNegate* Negate =
+		AND(Cast<UInputModifierNegate>(AModifierIsAppliedToAnActionMapping(Data, NewObject<UInputModifierNegate>(), TestContext, TestAction, EKeys::Gamepad_Left2D)));
+		check(Negate);
+		Negate->bX = false;
+		Negate->bY = true;
+		Negate->bZ = false;
+
+		WHEN(AKeyIsActuated(Data, EKeys::Gamepad_LeftX, 0.5f));
+		AND(AKeyIsActuated(Data, EKeys::Gamepad_LeftY, 0.5f));
+		AND(InputIsTicked(Data));
+		THEN(PressingKeyTriggersAction(Data, TestAction));
+
+		FVector Result = FInputTestHelper::GetTriggered<FVector>(Data, TestAction);
+		AND(TestEqual(TEXT("Negate Y only: X unchanged"), FX(Result), 0.5f));
+		AND(TestEqual(TEXT("Negate Y only: Y negated"), FY(Result), -0.5f));
+	}
+
+	// Test 3 - negate both X and Y
+	{
+		UControllablePlayer& Data =
+		GIVEN(ABasicModifierExtTest(this, EInputActionValueType::Axis2D));
+		AND(AnActionIsMappedToAKey(Data, TestContext, TestAction, EKeys::Gamepad_Left2D));
+
+		UInputModifierNegate* Negate =
+		AND(Cast<UInputModifierNegate>(AModifierIsAppliedToAnActionMapping(Data, NewObject<UInputModifierNegate>(), TestContext, TestAction, EKeys::Gamepad_Left2D)));
+		check(Negate);
+		Negate->bX = true;
+		Negate->bY = true;
+		Negate->bZ = false;
+
+		WHEN(AKeyIsActuated(Data, EKeys::Gamepad_LeftX, 0.5f));
+		AND(AKeyIsActuated(Data, EKeys::Gamepad_LeftY, 0.5f));
+		AND(InputIsTicked(Data));
+		THEN(PressingKeyTriggersAction(Data, TestAction));
+
+		FVector Result = FInputTestHelper::GetTriggered<FVector>(Data, TestAction);
+		AND(TestEqual(TEXT("Negate X+Y: X negated"), FX(Result), -0.5f));
+		AND(TestEqual(TEXT("Negate X+Y: Y negated"), FY(Result), -0.5f));
+	}
+
+	// Test 4 - no axes negated (all false): pass-through
+	{
+		UControllablePlayer& Data =
+		GIVEN(ABasicModifierExtTest(this, EInputActionValueType::Axis2D));
+		AND(AnActionIsMappedToAKey(Data, TestContext, TestAction, EKeys::Gamepad_Left2D));
+
+		UInputModifierNegate* Negate =
+		AND(Cast<UInputModifierNegate>(AModifierIsAppliedToAnActionMapping(Data, NewObject<UInputModifierNegate>(), TestContext, TestAction, EKeys::Gamepad_Left2D)));
+		check(Negate);
+		Negate->bX = false;
+		Negate->bY = false;
+		Negate->bZ = false;
+
+		WHEN(AKeyIsActuated(Data, EKeys::Gamepad_LeftX, 0.5f));
+		AND(AKeyIsActuated(Data, EKeys::Gamepad_LeftY, 0.5f));
+		AND(InputIsTicked(Data));
+		THEN(PressingKeyTriggersAction(Data, TestAction));
+
+		FVector Result = FInputTestHelper::GetTriggered<FVector>(Data, TestAction);
+		AND(TestEqual(TEXT("No negate: X unchanged"), FX(Result), 0.5f));
+		AND(TestEqual(TEXT("No negate: Y unchanged"), FY(Result), 0.5f));
+	}
+
+	return true;
+}
+
+// ============================================================
+// Dead Zone - Radial type
+// ============================================================
+
+IMPLEMENT_SIMPLE_AUTOMATION_TEST(FInputModifierDeadZoneRadialTest, "Input.Modifiers.DeadZone.Radial", BasicModifierExtTestFlags)
+
+bool FInputModifierDeadZoneRadialTest::RunTest(const FString& Parameters)
+{
+	// Radial dead zone applied at the action level so it sees the full assembled 2D vector.
+	UControllablePlayer& Data =
+	GIVEN(ABasicModifierExtTest(this, EInputActionValueType::Axis2D));
+	AND(AnActionIsMappedToAKey(Data, TestContext, TestAction, EKeys::Gamepad_Left2D));
+
+	UInputModifierDeadZone* DeadZone =
+	AND(Cast<UInputModifierDeadZone>(AModifierIsAppliedToAnAction(Data, NewObject<UInputModifierDeadZone>(), TestAction)));
+	check(DeadZone);
+	DeadZone->LowerThreshold = 0.2f;
+	DeadZone->UpperThreshold = 1.0f;
+	DeadZone->Type = EDeadZoneType::Radial;
+
+	// Test 1 - Input below the radial threshold: magnitude 0.1 < 0.2 -> no trigger
+	WHEN(AKeyIsActuated(Data, EKeys::Gamepad_LeftX, 0.1f));
+	AND(AKeyIsActuated(Data, EKeys::Gamepad_LeftY, 0.f));
+	AND(InputIsTicked(Data));
+	THEN(PressingKeyDoesNotTrigger(Data, TestAction));
+
+	// Test 2 - Input above the radial threshold: LeftX=0.3, LeftY=0.4 -> magnitude=0.5 > 0.2
+	// Expected: normal=(0.6,0.8), scaled_mag = (0.5-0.2)/(1.0-0.2) = 0.375
+	// Result = (0.6*0.375, 0.8*0.375) = (0.225, 0.3)
+	const float InputX = 0.3f;
+	const float InputY = 0.4f;
+	const float Magnitude = FMath::Sqrt(InputX * InputX + InputY * InputY); // 0.5
+	const float ScaledMag = FMath::Min(1.f, FMath::Max(0.f, Magnitude - DeadZone->LowerThreshold) / (DeadZone->UpperThreshold - DeadZone->LowerThreshold));
+	const float ExpectedX = (InputX / Magnitude) * ScaledMag;
+	const float ExpectedY = (InputY / Magnitude) * ScaledMag;
+
+	WHEN(AKeyIsActuated(Data, EKeys::Gamepad_LeftX, InputX));
+	AND(AKeyIsActuated(Data, EKeys::Gamepad_LeftY, InputY));
+	AND(InputIsTicked(Data));
+	THEN(PressingKeyTriggersAction(Data, TestAction));
+
+	FVector Result = FInputTestHelper::GetTriggered<FVector>(Data, TestAction);
+	AND(TestEqual(TEXT("Radial deadzone: X component"), FX(Result), ExpectedX));
+	AND(TestEqual(TEXT("Radial deadzone: Y component"), FY(Result), ExpectedY));
+
+	// Test 3 - Input at exactly the upper threshold: magnitude 1.0, scaled_mag = 1.0, result unchanged
+	WHEN(AKeyIsActuated(Data, EKeys::Gamepad_LeftX, 0.6f));
+	AND(AKeyIsActuated(Data, EKeys::Gamepad_LeftY, 0.8f));
+	AND(InputIsTicked(Data));
+	THEN(HoldingKeyTriggersAction(Data, TestAction));
+
+	Result = FInputTestHelper::GetTriggered<FVector>(Data, TestAction);
+	AND(TestEqual(TEXT("Radial deadzone: at upper threshold X"), FX(Result), 0.6f));
+	AND(TestEqual(TEXT("Radial deadzone: at upper threshold Y"), FY(Result), 0.8f));
+
+	// Test 4 - Release should stop triggering
+	WHEN(AKeyIsReleased(Data, EKeys::Gamepad_LeftX));
+	AND(AKeyIsReleased(Data, EKeys::Gamepad_LeftY));
+	AND(InputIsTicked(Data));
+	THEN(ReleasingKeyTriggersCompleted(Data, TestAction));
+
+	return true;
+}
+
+// ============================================================
+// Response Curve Exponential
+// ============================================================
+
+IMPLEMENT_SIMPLE_AUTOMATION_TEST(FInputModifierResponseCurveExponentialTest, "Input.Modifiers.ResponseCurveExponential", BasicModifierExtTestFlags)
+
+bool FInputModifierResponseCurveExponentialTest::RunTest(const FString& Parameters)
+{
+	UControllablePlayer& Data =
+	GIVEN(ABasicModifierExtTest(this, EInputActionValueType::Axis1D));
+	AND(AnActionIsMappedToAKey(Data, TestContext, TestAction, TestAxis));
+
+	UInputModifierResponseCurveExponential* Curve =
+	AND(Cast<UInputModifierResponseCurveExponential>(AModifierIsAppliedToAnActionMapping(Data, NewObject<UInputModifierResponseCurveExponential>(), TestContext, TestAction, TestAxis)));
+	check(Curve);
+
+	// Test 1 - Exponent 1.0 (identity): output == input
+	Curve->CurveExponent = FVector(1.f, 1.f, 1.f);
+	WHEN(AKeyIsActuated(Data, TestAxis, 0.5f));
+	AND(InputIsTicked(Data));
+	THEN(PressingKeyTriggersAction(Data, TestAction));
+	AND(TestEqual(TEXT("Exponent 1.0: identity"), GetModExtActionValue(Data, TestAction), 0.5f));
+
+	// Test 2 - Exponent 2.0: output = sign(x) * pow(|x|, 2) = 0.25 for input 0.5
+	Curve->CurveExponent = FVector(2.f, 1.f, 1.f);
+	WHEN(AKeyIsActuated(Data, TestAxis, 0.5f));
+	AND(InputIsTicked(Data));
+	THEN(HoldingKeyTriggersAction(Data, TestAction));
+	AND(TestEqual(TEXT("Exponent 2.0 positive input"), GetModExtActionValue(Data, TestAction), 0.25f));
+
+	// Test 3 - Exponent 2.0 with negative input: sign is preserved (-0.25 for -0.5)
+	WHEN(AKeyIsActuated(Data, TestAxis, -0.5f));
+	AND(InputIsTicked(Data));
+	THEN(HoldingKeyTriggersAction(Data, TestAction));
+	AND(TestEqual(TEXT("Exponent 2.0 negative input"), GetModExtActionValue(Data, TestAction), -0.25f));
+
+	// Test 4 - Exponent 2.0 at input 1.0: pow(1, 2) = 1.0 (boundary unchanged)
+	WHEN(AKeyIsActuated(Data, TestAxis, 1.0f));
+	AND(InputIsTicked(Data));
+	THEN(HoldingKeyTriggersAction(Data, TestAction));
+	AND(TestEqual(TEXT("Exponent 2.0 at max input"), GetModExtActionValue(Data, TestAction), 1.0f));
+
+	// Test 5 - Exponent 0.5 (square root): output = sqrt(0.25) = 0.5 for input 0.25 (expands the curve)
+	Curve->CurveExponent = FVector(0.5f, 1.f, 1.f);
+	WHEN(AKeyIsActuated(Data, TestAxis, 0.25f));
+	AND(InputIsTicked(Data));
+	THEN(HoldingKeyTriggersAction(Data, TestAction));
+	AND(TestEqual(TEXT("Exponent 0.5 (sqrt): value expanded"), GetModExtActionValue(Data, TestAction), 0.5f));
+
+	// Release
+	WHEN(AKeyIsReleased(Data, TestAxis));
+	AND(InputIsTicked(Data));
+	THEN(ReleasingKeyTriggersCompleted(Data, TestAction));
+
+	return true;
+}
+
+// ============================================================
+// Scale By Delta Time
+// ============================================================
+
+IMPLEMENT_SIMPLE_AUTOMATION_TEST(FInputModifierScaleByDeltaTimeTest, "Input.Modifiers.ScaleByDeltaTime", BasicModifierExtTestFlags)
+
+bool FInputModifierScaleByDeltaTimeTest::RunTest(const FString& Parameters)
+{
+	UControllablePlayer& Data =
+	GIVEN(ABasicModifierExtTest(this, EInputActionValueType::Axis1D));
+	AND(AnActionIsMappedToAKey(Data, TestContext, TestAction, TestAxis));
+	AND(AModifierIsAppliedToAnActionMapping(Data, NewObject<UInputModifierScaleByDeltaTime>(), TestContext, TestAction, TestAxis));
+
+	// Test 1 - Input 1.0 with DeltaTime=1/60: output = 1.0 * (1/60)
+	const float FrameTime = 1.f / 60.f;
+	WHEN(AKeyIsActuated(Data, TestAxis, 1.0f));
+	AND(InputIsTicked(Data, FrameTime));
+	THEN(PressingKeyTriggersAction(Data, TestAction));
+	AND(TestEqual(TEXT("Scale by DeltaTime at 60fps"), GetModExtActionValue(Data, TestAction), 1.0f * FrameTime));
+
+	// Test 2 - Same input with DeltaTime=1/30 (half frame rate): output doubles
+	const float HalfRate = 1.f / 30.f;
+	WHEN(InputIsTicked(Data, HalfRate));
+	THEN(HoldingKeyTriggersAction(Data, TestAction));
+	AND(TestEqual(TEXT("Scale by DeltaTime at 30fps"), GetModExtActionValue(Data, TestAction), 1.0f * HalfRate));
+
+	// Test 3 - Different input magnitude: 0.5 * (1/60)
+	WHEN(AKeyIsActuated(Data, TestAxis, 0.5f));
+	AND(InputIsTicked(Data, FrameTime));
+	THEN(HoldingKeyTriggersAction(Data, TestAction));
+	AND(TestEqual(TEXT("Scale by DeltaTime: half magnitude"), GetModExtActionValue(Data, TestAction), 0.5f * FrameTime));
+
+	// Release
+	WHEN(AKeyIsReleased(Data, TestAxis));
+	AND(InputIsTicked(Data, FrameTime));
+	THEN(ReleasingKeyTriggersCompleted(Data, TestAction));
+
+	return true;
+}
+
+// ============================================================
+// Modifier application order: mapping-level runs before action-level
+// ============================================================
+
+IMPLEMENT_SIMPLE_AUTOMATION_TEST(FInputModifierChainOrderTest, "Input.Modifiers.ApplicationOrder", BasicModifierExtTestFlags)
+
+bool FInputModifierChainOrderTest::RunTest(const FString& Parameters)
+{
+	// Mapping-level modifier: Scalar(3x)
+	// Action-level modifier: DeadZone(LowerThreshold=0.5)
+	// Input: 0.2
+	//
+	// If mapping runs first:  0.2 * 3 = 0.6 -> deadzone((0.6-0.5)/(1.0-0.5)) = 0.2 -> TRIGGERED with value 0.2
+	// If action runs first:   deadzone(0.2) = 0.0 (below threshold) -> 0.0 * 3 = 0.0 -> NOT TRIGGERED
+	//
+	// The expected UE behavior is mapping modifiers before action modifiers.
+	UControllablePlayer& Data =
+	GIVEN(ABasicModifierExtTest(this, EInputActionValueType::Axis1D));
+	AND(AnActionIsMappedToAKey(Data, TestContext, TestAction, TestAxis));
+
+	// Pre-configure the Scalar so the deep-copy during RebuildControlMappings captures the 3x value.
+	// Each ControlMappingsAreRebuilt call deep-copies context modifiers, invalidating any prior live pointer.
+	// By setting the value on the original BEFORE applying, subsequent rebuilds produce correct duplicates.
+	UInputModifierScalar* OrigScalar = NewObject<UInputModifierScalar>();
+	OrigScalar->Scalar = FVector(3.f, 3.f, 3.f);
+	AND(AModifierIsAppliedToAnActionMapping(Data, OrigScalar, TestContext, TestAction, TestAxis));
+
+	// Pre-configure the DeadZone similarly before its rebuild.
+	UInputModifierDeadZone* OrigDeadZone = NewObject<UInputModifierDeadZone>();
+	OrigDeadZone->LowerThreshold = 0.5f;
+	OrigDeadZone->UpperThreshold = 1.0f;
+	OrigDeadZone->Type = EDeadZoneType::Axial;
+	AND(AModifierIsAppliedToAnAction(Data, OrigDeadZone, TestAction));
+
+	// After the last rebuild, get a fresh pointer to the live Scalar for use in sub-tests.
+	FEnhancedActionKeyMapping* LiveMap = FInputTestHelper::FindLiveActionMapping(Data, TestAction, TestAxis);
+	check(LiveMap && !LiveMap->Modifiers.IsEmpty());
+	UInputModifierScalar* LiveScalar = Cast<UInputModifierScalar>(LiveMap->Modifiers[0]);
+	check(LiveScalar);
+
+	// Input 0.2: must be scaled to 0.6 by the mapping modifier BEFORE hitting the deadzone.
+	WHEN(AKeyIsActuated(Data, TestAxis, 0.2f));
+	AND(InputIsTicked(Data));
+
+	// Should trigger because 0.2 -> 0.6 (above deadzone lower threshold of 0.5)
+	THEN(PressingKeyTriggersAction(Data, TestAction));
+	// Value after deadzone: (0.6-0.5)/(1.0-0.5) = 0.2
+	AND(TestEqual(TEXT("Mapping modifier ran before action modifier"), GetModExtActionValue(Data, TestAction), 0.2f));
+
+	// Test 2 - Set scalar to identity so raw 0.2 hits deadzone directly.
+	// LastState=Triggered (from Test 1), 0.2 now killed by deadzone -> TriggerState=None -> Triggered->None=Completed.
+	LiveScalar->Scalar = FVector(1.f, 1.f, 1.f);
+	WHEN(AKeyIsActuated(Data, TestAxis, 0.2f));
+	AND(InputIsTicked(Data));
+	THEN(HoldingKeyTriggersCompleted(Data, TestAction)); // Triggered->None transition when value drops below deadzone
+
+	// Restore scale: action returns from None -> StartedAndTriggered fires.
+	LiveScalar->Scalar = FVector(3.f, 3.f, 3.f);
+	WHEN(AKeyIsActuated(Data, TestAxis, 0.2f));
+	AND(InputIsTicked(Data));
+	THEN(PressingKeyTriggersAction(Data, TestAction));
+
+	return true;
+}
diff --git a/Engine/Plugins/EnhancedInput/Source/InputEditor/Private/Tests/InputSystemExtendedTest.cpp b/Engine/Plugins/EnhancedInput/Source/InputEditor/Private/Tests/InputSystemExtendedTest.cpp
new file mode 100644
--- /dev/null
+++ b/Engine/Plugins/EnhancedInput/Source/InputEditor/Private/Tests/InputSystemExtendedTest.cpp
@@ -0,0 +1,265 @@
+// Copyright Epic Games, Inc. All Rights Reserved.
+
+#include "InputTestFramework.h"
+#include "Misc/AutomationTest.h"
+#include "InputAction.h"
+#include "InputModifiers.h"
+#include "EnhancedActionKeyMapping.h"
+
+// Extended system tests covering AccumulationBehavior, Axis3D value flow,
+// and the modifier application order between mapping-level and action-level.
+
+constexpr auto BasicSystemExtTestFlags = EAutomationTestFlags::EditorContext | EAutomationTestFlags::EngineFilter;
+
+static UControllablePlayer& ABasicSystemExtTest(FAutomationTestBase* Test, EInputActionValueType ForValueType)
+{
+	UWorld* World =
+	GIVEN(AnEmptyWorld());
+
+	UControllablePlayer& Data =
+	AND(AControllablePlayer(World));
+	Test->TestTrue(TEXT("Controllable Player is valid"), Data.IsValid());
+
+	AND(AnInputContextIsAppliedToAPlayer(Data, TestContext, 0));
+	AND(AnInputAction(Data, TestAction, ForValueType));
+
+	return Data;
+}
+
+// ============================================================
+// AccumulationBehavior: TakeHighestAbsoluteValue vs Cumulative
+// ============================================================
+
+IMPLEMENT_SIMPLE_AUTOMATION_TEST(FInputAccumulationBehaviorTest, "Input.System.AccumulationBehavior", BasicSystemExtTestFlags)
+
+bool FInputAccumulationBehaviorTest::RunTest(const FString& Parameters)
+{
+	// Two 1D keys mapped to the same Axis1D action.
+	// Pressing both simultaneously lets us test how their values are combined.
+	UControllablePlayer& Data =
+	GIVEN(ABasicSystemExtTest(this, EInputActionValueType::Axis1D));
+	AND(AnActionIsMappedToAKey(Data, TestContext, TestAction, TestAxis));
+	AND(AnActionIsMappedToAKey(Data, TestContext, TestAction, TestAxis2));
+
+	UInputAction* Action = FInputTestHelper::FindAction(Data, TestAction);
+	TestNotNull(TEXT("Action found"), Action);
+
+	// -------------------------------------------------------
+	// Test 1 - TakeHighestAbsoluteValue (default):
+	//   Both keys pressed: 0.3 and 0.8 -> result = 0.8 (highest absolute value wins)
+	// -------------------------------------------------------
+	Action->AccumulationBehavior = EInputActionAccumulationBehavior::TakeHighestAbsoluteValue;
+
+	WHEN(AKeyIsActuated(Data, TestAxis, 0.3f));
+	AND(AKeyIsActuated(Data, TestAxis2, 0.8f));
+	AND(InputIsTicked(Data));
+	THEN(PressingKeyTriggersAction(Data, TestAction));
+	AND(TestEqual(TEXT("TakeHighest: 0.8 beats 0.3"), FInputTestHelper::GetActionData(Data, TestAction).GetValue().Get<float>(), 0.8f));
+
+	// Test 1b - Higher absolute value wins regardless of evaluation order
+	WHEN(AKeyIsActuated(Data, TestAxis, 0.8f));
+	AND(AKeyIsActuated(Data, TestAxis2, 0.3f));
+	AND(InputIsTicked(Data));
+	THEN(HoldingKeyTriggersAction(Data, TestAction));
+	AND(TestEqual(TEXT("TakeHighest: 0.8 beats 0.3 (reversed)"), FInputTestHelper::GetActionData(Data, TestAction).GetValue().Get<float>(), 0.8f));
+
+	// Test 1c - Negative value with higher absolute magnitude wins
+	WHEN(AKeyIsActuated(Data, TestAxis, -0.8f));
+	AND(AKeyIsActuated(Data, TestAxis2, 0.3f));
+	AND(InputIsTicked(Data));
+	THEN(HoldingKeyTriggersAction(Data, TestAction));
+	AND(TestEqual(TEXT("TakeHighest: -0.8 beats 0.3 (sign preserved)"), FInputTestHelper::GetActionData(Data, TestAction).GetValue().Get<float>(), -0.8f));
+
+	// -------------------------------------------------------
+	// Test 2 - Cumulative:
+	//   Both keys pressed: 0.3 and 0.5 -> result = 0.8 (values summed)
+	// -------------------------------------------------------
+	Action->AccumulationBehavior = EInputActionAccumulationBehavior::Cumulative;
+
+	WHEN(AKeyIsActuated(Data, TestAxis, 0.3f));
+	AND(AKeyIsActuated(Data, TestAxis2, 0.5f));
+	AND(InputIsTicked(Data));
+	THEN(HoldingKeyTriggersAction(Data, TestAction));
+	AND(TestEqual(TEXT("Cumulative: 0.3 + 0.5 = 0.8"), FInputTestHelper::GetActionData(Data, TestAction).GetValue().Get<float>(), 0.8f));
+
+	// Test 2b - Opposite-sign values cancel out (WASD-style: W=+1, S=-1 -> 0)
+	// NOTE: SetStateForNoTriggers is called per-mapping with each mapping's INDIVIDUAL value.
+	// Both ±0.5 mappings are individually non-zero, so each contributes Triggered to the tracker.
+	// The accumulated sum is 0.0 but the trigger state is still Triggered (actual system behavior).
+	WHEN(AKeyIsActuated(Data, TestAxis, 0.5f));
+	AND(AKeyIsActuated(Data, TestAxis2, -0.5f));
+	AND(InputIsTicked(Data));
+	THEN(HoldingKeyTriggersAction(Data, TestAction)); // both mappings individually non-zero -> Triggered
+	AND(TestEqual(TEXT("Cumulative: +0.5 + -0.5 = 0"), FInputTestHelper::GetActionData(Data, TestAction).GetValue().Get<float>(), 0.0f));
+
+	// Test 2c - Single key: Cumulative with one key behaves the same as TakeHighest
+	WHEN(AKeyIsReleased(Data, TestAxis2));
+	AND(AKeyIsActuated(Data, TestAxis, 0.7f));
+	AND(InputIsTicked(Data));
+	THEN(HoldingKeyTriggersAction(Data, TestAction));
+	AND(TestEqual(TEXT("Cumulative: single key"), FInputTestHelper::GetActionData(Data, TestAction).GetValue().Get<float>(), 0.7f));
+
+	return true;
+}
+
+// ============================================================
+// Axis3D action value flow
+// ============================================================
+
+IMPLEMENT_SIMPLE_AUTOMATION_TEST(FInputAxis3DActionTest, "Input.System.Axis3D", BasicSystemExtTestFlags)
+
+bool FInputAxis3DActionTest::RunTest(const FString& Parameters)
+{
+	// A 1D analog key maps to the X component of an Axis3D action.
+	// We verify the value is correctly placed in X and that Y/Z are zero.
+	UControllablePlayer& Data =
+	GIVEN(ABasicSystemExtTest(this, EInputActionValueType::Axis3D));
+	AND(AnActionIsMappedToAKey(Data, TestContext, TestAction, TestAxis));
+
+	// Test 1 - Single 1D key: value lands on X, Y and Z are zero
+	WHEN(AKeyIsActuated(Data, TestAxis, 0.7f));
+	AND(InputIsTicked(Data));
+	THEN(PressingKeyTriggersAction(Data, TestAction));
+
+	FVector Result = FInputTestHelper::GetTriggered<FVector>(Data, TestAction);
+	AND(TestEqual(TEXT("Axis3D: X from 1D key"), (float)Result.X, 0.7f));
+	AND(TestEqual(TEXT("Axis3D: Y is zero"), (float)Result.Y, 0.f));
+	AND(TestEqual(TEXT("Axis3D: Z is zero"), (float)Result.Z, 0.f));
+
+	// Test 2 - Modify the key value and verify the updated X
+	WHEN(AKeyIsActuated(Data, TestAxis, -0.4f));
+	AND(InputIsTicked(Data));
+	THEN(HoldingKeyTriggersAction(Data, TestAction));
+
+	Result = FInputTestHelper::GetTriggered<FVector>(Data, TestAction);
+	AND(TestEqual(TEXT("Axis3D: negative X"), (float)Result.X, -0.4f));
+
+	// Test 3 - SwizzleAxis ZYX at the mapping level: X value moves to Z component
+	//          Input FVector(0.5, 0, 0) -> after ZYX (swap X,Z) -> FVector(0, 0, 0.5)
+	AND(AnActionIsMappedToAKey(Data, TestContext, TestAction, TestAxis2)); // map a second key for this sub-test
+	// Pre-configure before applying so DeepCopyPtrArray captures ZYX order.
+	UInputModifierSwizzleAxis* OrigSwizzle = NewObject<UInputModifierSwizzleAxis>();
+	OrigSwizzle->Order = EInputAxisSwizzle::ZYX;
+	AND(AModifierIsAppliedToAnActionMapping(Data, OrigSwizzle, TestContext, TestAction, TestAxis2));
+
+	// Release the first key so only TestAxis2 drives the action
+	WHEN(AKeyIsReleased(Data, TestAxis));
+	AND(AKeyIsActuated(Data, TestAxis2, 0.5f));
+	AND(InputIsTicked(Data));
+	THEN(HoldingKeyTriggersAction(Data, TestAction));
+
+	Result = FInputTestHelper::GetTriggered<FVector>(Data, TestAction);
+	AND(TestEqual(TEXT("Axis3D + SwizzleZYX: X is zero"), (float)Result.X, 0.f));
+	AND(TestEqual(TEXT("Axis3D + SwizzleZYX: Y is zero"), (float)Result.Y, 0.f));
+	AND(TestEqual(TEXT("Axis3D + SwizzleZYX: Z receives value"), (float)Result.Z, 0.5f));
+
+	// Test 4 - Release: Completed
+	WHEN(AKeyIsReleased(Data, TestAxis2));
+	AND(InputIsTicked(Data));
+	THEN(ReleasingKeyTriggersCompleted(Data, TestAction));
+
+	return true;
+}
+
+// ============================================================
+// Modifier application order: mapping-level vs action-level
+// on independent axes to demonstrate ordering through the pipeline.
+// ============================================================
+
+IMPLEMENT_SIMPLE_AUTOMATION_TEST(FInputModifierApplicationOrderTest, "Input.System.ModifierApplicationOrder", BasicSystemExtTestFlags)
+
+bool FInputModifierApplicationOrderTest::RunTest(const FString& Parameters)
+{
+	// This test establishes that:
+	//   1. Modifiers in the KEY MAPPING chain run before ACTION-LEVEL modifiers.
+	//   2. Within a chain (mapping or action level), modifiers run in array order.
+	//
+	// Strategy: use non-commutative operations so the two orderings produce
+	// different results.
+	//
+	// Mapping-level: Scalar(3)
+	// Action-level:  DeadZone(LowerThreshold=0.5, UpperThreshold=1.0)
+	// Input: 0.2
+	//
+	//   Correct order (mapping first):  0.2 * 3 = 0.6 -> deadzone -> (0.6-0.5)/0.5 = 0.2 -> TRIGGERED
+	//   Reversed order (action first):  deadzone(0.2) = 0.0 (below 0.5) -> * 3 = 0.0 -> NOT TRIGGERED
+	//
+	// Sub-test A verifies the correct order fires the trigger.
+	// Sub-test B removes the mapping-level scalar to confirm a raw 0.2 input is killed
+	// by the 0.5 deadzone, making the contrast explicit.
+
+	UControllablePlayer& Data =
+	GIVEN(ABasicSystemExtTest(this, EInputActionValueType::Axis1D));
+	AND(AnActionIsMappedToAKey(Data, TestContext, TestAction, TestAxis));
+
+	// Pre-configure modifiers before applying so the deep-copy in RebuildControlMappings
+	// captures the correct values. Post-apply pointers are duplicates and will be stale
+	// after the next rebuild, so we fetch a fresh live pointer after all Apply calls.
+	UInputModifierScalar* OrigScalar = NewObject<UInputModifierScalar>();
+	OrigScalar->Scalar = FVector(3.f, 3.f, 3.f);
+	AND(AModifierIsAppliedToAnActionMapping(Data, OrigScalar, TestContext, TestAction, TestAxis));
+
+	UInputModifierDeadZone* OrigDeadZone = NewObject<UInputModifierDeadZone>();
+	OrigDeadZone->LowerThreshold = 0.5f;
+	OrigDeadZone->UpperThreshold = 1.0f;
+	OrigDeadZone->Type = EDeadZoneType::Axial;
+	AND(AModifierIsAppliedToAnAction(Data, OrigDeadZone, TestAction));
+
+	// Fetch the live Scalar duplicate (created by DeepCopyPtrArray during the last rebuild).
+	// This is the pointer we can safely mutate between ticks without triggering another rebuild.
+	FEnhancedActionKeyMapping* LiveMap = FInputTestHelper::FindLiveActionMapping(Data, TestAction, TestAxis);
+	check(LiveMap && !LiveMap->Modifiers.IsEmpty());
+	UInputModifierScalar* LiveScalar = Cast<UInputModifierScalar>(LiveMap->Modifiers[0]);
+	check(LiveScalar);
+
+	// Sub-test A: input 0.2 should survive because the mapping scalar runs first (0.2*3=0.6 > 0.5 threshold)
+	WHEN(AKeyIsActuated(Data, TestAxis, 0.2f));
+	AND(InputIsTicked(Data));
+	THEN(PressingKeyTriggersAction(Data, TestAction));
+	// Post-deadzone value: (0.6 - 0.5) / (1.0 - 0.5) = 0.2
+	AND(TestEqual(TEXT("Mapping modifier ran first: value is 0.2 post-deadzone"), FInputTestHelper::GetActionData(Data, TestAction).GetValue().Get<float>(), 0.2f));
+
+	// Sub-test B: set the mapping-level scalar to identity (1.0).
+	// Now raw 0.2 goes straight to the deadzone and is killed (0.2 < 0.5 threshold).
+	// State was Triggered from sub-test A; value drops to 0 -> Triggered->None=Completed fires.
+	LiveScalar->Scalar = FVector(1.f, 1.f, 1.f);
+
+	WHEN(AKeyIsActuated(Data, TestAxis, 0.2f));
+	AND(InputIsTicked(Data));
+	THEN(HoldingKeyTriggersCompleted(Data, TestAction)); // Triggered->None transition
+
+	// Sub-test C: restore scalar and verify it triggers again, confirming the ordering is stable.
+	LiveScalar->Scalar = FVector(3.f, 3.f, 3.f);
+
+	WHEN(AKeyIsActuated(Data, TestAxis, 0.2f));
+	AND(InputIsTicked(Data));
+	THEN(PressingKeyTriggersAction(Data, TestAction)); // None->Triggered=StartedAndTriggered
+
+	// -------------------------------------------------------
+	// Sub-test D: within a single level, modifiers run in array (insertion) order.
+	// Negate then Scalar(-2): input 1.0 -> negate -> -1.0 -> scalar(-2) -> 2.0
+	// Scalar(-2) then Negate: input 1.0 -> scalar(-2) -> -2.0 -> negate -> 2.0
+	// Same result for linear ops - but verifying two modifiers in a chain both fire.
+	// -------------------------------------------------------
+	// Use a fresh test setup to isolate this sub-test.
+	{
+		UControllablePlayer& ChainData =
+		GIVEN(ABasicSystemExtTest(this, EInputActionValueType::Axis1D));
+		AND(AnActionIsMappedToAKey(ChainData, TestContext, TestAction, TestAxis));
+
+		// Apply Negate first, then Scalar(-2) - both at mapping level.
+		// Pre-configure before applying so DeepCopyPtrArray captures the correct value.
+		AND(AModifierIsAppliedToAnActionMapping(ChainData, NewObject<UInputModifierNegate>(), TestContext, TestAction, TestAxis));
+		UInputModifierScalar* OrigChainScalar = NewObject<UInputModifierScalar>();
+		OrigChainScalar->Scalar = FVector(-2.f, -2.f, -2.f);
+		AND(AModifierIsAppliedToAnActionMapping(ChainData, OrigChainScalar, TestContext, TestAction, TestAxis));
+
+		// input 1.0 -> Negate -> -1.0 -> Scalar(-2) -> 2.0
+		WHEN(AKeyIsActuated(ChainData, TestAxis, 1.0f));
+		AND(InputIsTicked(ChainData));
+		THEN(PressingKeyTriggersAction(ChainData, TestAction));
+		AND(TestEqual(TEXT("Chain: Negate then Scalar(-2) applied in order"), FInputTestHelper::GetActionData(ChainData, TestAction).GetValue().Get<float>(), 2.0f));
+	}
+
+	return true;
+}
diff --git a/Engine/Plugins/EnhancedInput/Source/InputEditor/Private/Tests/InputTriggerExtendedTest.cpp b/Engine/Plugins/EnhancedInput/Source/InputEditor/Private/Tests/InputTriggerExtendedTest.cpp
new file mode 100644
--- /dev/null
+++ b/Engine/Plugins/EnhancedInput/Source/InputEditor/Private/Tests/InputTriggerExtendedTest.cpp
@@ -0,0 +1,503 @@
+// Copyright Epic Games, Inc. All Rights Reserved.
+
+#include "InputTestFramework.h"
+#include "Misc/AutomationTest.h"
+#include "InputTriggers.h"
+
+// Extended trigger integration tests: Pulse, Hold (OneShot), HoldAndRelease, Tap, and ChordBlocker.
+// These are integration-style tests using a full UControllablePlayer+world, which lets time-based
+// triggers receive proper DeltaTime through the input stack.
+
+constexpr auto BasicTriggerExtTestFlags = EAutomationTestFlags::EditorContext | EAutomationTestFlags::EngineFilter;
+
+// Local helper: creates world+player+context+action+key mapping and applies trigger T at the action level.
+// Returns the applied trigger instance cast to T*.
+template<typename T>
+static T* ABasicTriggerExtTest(FAutomationTestBase* Test, UControllablePlayer*& OutData)
+{
+	UWorld* World =
+	GIVEN(AnEmptyWorld());
+
+	UControllablePlayer& Data =
+	AND(AControllablePlayer(World));
+	Test->TestTrue(TEXT("Controllable Player is valid"), Data.IsValid());
+
+	AND(AnInputContextIsAppliedToAPlayer(Data, TestContext, 0));
+	AND(AnInputAction(Data, TestAction, EInputActionValueType::Boolean));
+	AND(AnActionIsMappedToAKey(Data, TestContext, TestAction, TestKey));
+
+	T* TriggerInstance = Cast<T>(ATriggerIsAppliedToAnAction(Data, NewObject<T>(), TestAction));
+	Test->TestNotNull(TEXT("Trigger instance created"), TriggerInstance);
+
+	OutData = &Data;
+	return TriggerInstance;
+}
+
+// ============================================================
+// Pulse trigger
+// ============================================================
+
+IMPLEMENT_SIMPLE_AUTOMATION_TEST(FInputIntegrationTriggerPulseTest, "Input.Integration.Triggers.Pulse", BasicTriggerExtTestFlags)
+
+bool FInputIntegrationTriggerPulseTest::RunTest(const FString& Parameters)
+{
+	const float FrameTime = 1.f / 60.f;
+	// Use a 30-frame interval. Setting Interval to 29.5 frames ensures the fire tick (frame 30)
+	// is strictly past the threshold with no floating-point ambiguity (avoids HeldDuration == Interval).
+	const int32 IntervalFrames = 30;
+	const float Interval = FrameTime * (IntervalFrames - 0.5f);
+
+	// -------------------------------------------------------
+	// Test 1 - bTriggerOnStart = true:
+	//   First tick while held -> immediately fires Triggered (HeldDuration > Interval*0 = 0).
+	//   Subsequent ticks before the next interval -> Ongoing (not triggered).
+	//   At the next interval threshold -> fires again.
+	// -------------------------------------------------------
+	{
+		UControllablePlayer* DataPtr = nullptr;
+		UInputTriggerPulse* Pulse = ABasicTriggerExtTest<UInputTriggerPulse>(this, DataPtr);
+		check(Pulse && DataPtr);
+		UControllablePlayer& Data = *DataPtr;
+
+		Pulse->bTriggerOnStart = true;
+		Pulse->Interval = Interval;
+		Pulse->TriggerLimit = 0; // unlimited
+
+		// First tick: HeldDuration = FrameTime > 0 -> triggers immediately
+		WHEN(AKeyIsActuated(Data, TestKey));
+		AND(InputIsTicked(Data, FrameTime));
+		THEN(PressingKeyTriggersAction(Data, TestAction));
+
+		// Mid-interval ticks: HeldDuration < Interval*1 -> Ongoing only
+		for (int32 Frame = 1; Frame < IntervalFrames - 1; ++Frame)
+		{
+			WHEN(InputIsTicked(Data, FrameTime));
+			THEN(HoldingKeyTriggersOngoing(Data, TestAction));
+		}
+
+		// At/past the interval threshold (frame IntervalFrames): triggers again
+		WHEN(InputIsTicked(Data, FrameTime));
+		THEN(HoldingKeyTriggersAction(Data, TestAction));
+
+		// Releasing transitions to Completed
+		WHEN(AKeyIsReleased(Data, TestKey));
+		AND(InputIsTicked(Data, FrameTime));
+		THEN(ReleasingKeyTriggersCompleted(Data, TestAction));
+	}
+
+	// -------------------------------------------------------
+	// Test 2 - bTriggerOnStart = false:
+	//   First tick -> Ongoing (not triggered yet).
+	//   At the first interval threshold -> fires Triggered.
+	//   Second interval -> fires again.
+	// -------------------------------------------------------
+	{
+		UControllablePlayer* DataPtr = nullptr;
+		UInputTriggerPulse* Pulse = ABasicTriggerExtTest<UInputTriggerPulse>(this, DataPtr);
+		check(Pulse && DataPtr);
+		UControllablePlayer& Data = *DataPtr;
+
+		Pulse->bTriggerOnStart = false;
+		Pulse->Interval = Interval;
+		Pulse->TriggerLimit = 0;
+
+		// First tick: HeldDuration = FrameTime < Interval -> Ongoing.
+		// State transition: None->Ongoing fires Started (not Ongoing).
+		WHEN(AKeyIsActuated(Data, TestKey));
+		AND(InputIsTicked(Data, FrameTime));
+		THEN(PressingKeyTriggersStarted(Data, TestAction));
+
+		// Ticks up to but not including the interval frame
+		for (int32 Frame = 1; Frame < IntervalFrames - 1; ++Frame)
+		{
+			WHEN(InputIsTicked(Data, FrameTime));
+			THEN(HoldingKeyTriggersOngoing(Data, TestAction));
+		}
+
+		// At the interval threshold: fires first Triggered
+		WHEN(InputIsTicked(Data, FrameTime));
+		THEN(HoldingKeyTriggersAction(Data, TestAction));
+
+		// Mid-second-interval: Ongoing again
+		for (int32 Frame = 0; Frame < IntervalFrames - 1; ++Frame)
+		{
+			WHEN(InputIsTicked(Data, FrameTime));
+			THEN(HoldingKeyTriggersOngoing(Data, TestAction));
+		}
+
+		// At the second interval: fires again
+		WHEN(InputIsTicked(Data, FrameTime));
+		THEN(HoldingKeyTriggersAction(Data, TestAction));
+	}
+
+	// -------------------------------------------------------
+	// Test 3 - TriggerLimit = 2: fires at most twice, then goes silent (None)
+	// -------------------------------------------------------
+	{
+		UControllablePlayer* DataPtr = nullptr;
+		UInputTriggerPulse* Pulse = ABasicTriggerExtTest<UInputTriggerPulse>(this, DataPtr);
+		check(Pulse && DataPtr);
+		UControllablePlayer& Data = *DataPtr;
+
+		Pulse->bTriggerOnStart = true;
+		Pulse->Interval = Interval;
+		Pulse->TriggerLimit = 2;
+
+		// First fire: frame 1 (bTriggerOnStart)
+		WHEN(AKeyIsActuated(Data, TestKey));
+		AND(InputIsTicked(Data, FrameTime));
+		THEN(PressingKeyTriggersAction(Data, TestAction));
+
+		// Ticks through to second fire
+		for (int32 Frame = 1; Frame < IntervalFrames - 1; ++Frame)
+		{
+			WHEN(InputIsTicked(Data, FrameTime));
+			THEN(HoldingKeyTriggersOngoing(Data, TestAction));
+		}
+
+		// Second fire at interval threshold (TriggerCount will reach limit here)
+		WHEN(InputIsTicked(Data, FrameTime));
+		THEN(HoldingKeyTriggersAction(Data, TestAction));
+
+		// After TriggerLimit reached, Pulse returns None. LastState was Triggered -> Triggered->None=Completed.
+		WHEN(InputIsTicked(Data, FrameTime));
+		THEN(HoldingKeyTriggersCompleted(Data, TestAction));
+
+		// Further ticks: None->None -> no event.
+		WHEN(InputIsTicked(Data, FrameTime));
+		THEN(HoldingKeyDoesNotTrigger(Data, TestAction));
+	}
+
+	return true;
+}
+
+// ============================================================
+// Hold trigger - bIsOneShot = true
+// ============================================================
+
+IMPLEMENT_SIMPLE_AUTOMATION_TEST(FInputIntegrationTriggerHoldOneShotTest, "Input.Integration.Triggers.Hold.OneShot", BasicTriggerExtTestFlags)
+
+bool FInputIntegrationTriggerHoldOneShotTest::RunTest(const FString& Parameters)
+{
+	const float FrameTime = 1.f / 60.f;
+	const int32 HoldFrames = 30;
+
+	UControllablePlayer* DataPtr = nullptr;
+	UInputTriggerHold* HoldTrigger = ABasicTriggerExtTest<UInputTriggerHold>(this, DataPtr);
+	check(HoldTrigger && DataPtr);
+	UControllablePlayer& Data = *DataPtr;
+
+	HoldTrigger->HoldTimeThreshold = FrameTime * HoldFrames;
+	HoldTrigger->bIsOneShot = true;
+
+	// -------------------------------------------------------
+	// Test 1 - OneShot: fires exactly once at threshold, then stays silent while held.
+	// -------------------------------------------------------
+
+	// Pressing
+	WHEN(AKeyIsActuated(Data, TestKey));
+	AND(InputIsTicked(Data, FrameTime));
+	THEN(PressingKeyTriggersStarted(Data, TestAction));
+
+	// Holding until threshold
+	for (int32 Frame = 1; Frame < HoldFrames - 1; ++Frame)
+	{
+		WHEN(InputIsTicked(Data, FrameTime));
+		THEN(HoldingKeyTriggersOngoing(Data, TestAction));
+	}
+
+	// Reaching threshold: fires once
+	WHEN(InputIsTicked(Data, FrameTime));
+	THEN(HoldingKeyTriggersAction(Data, TestAction));
+
+	// Tick after OneShot: Hold (bIsOneShot=true, already triggered) returns None.
+	// LastState=Triggered -> Triggered->None=Completed fires.
+	WHEN(InputIsTicked(Data, FrameTime));
+	THEN(HoldingKeyTriggersCompleted(Data, TestAction));
+
+	// Further ticks while held: None->None -> no event.
+	WHEN(InputIsTicked(Data, FrameTime));
+	THEN(HoldingKeyDoesNotTrigger(Data, TestAction));
+
+	// Releasing: state was already None after Completed, key release produces None->None -> no event.
+	WHEN(AKeyIsReleased(Data, TestKey));
+	AND(InputIsTicked(Data, FrameTime));
+	THEN(ReleasingKeyDoesNotTrigger(Data, TestAction));
+
+	// -------------------------------------------------------
+	// Test 2 - OneShot reset: pressing again after release should fire once more.
+	// -------------------------------------------------------
+	InputIsTicked(Data, FrameTime); // clear state
+
+	WHEN(AKeyIsActuated(Data, TestKey));
+	AND(InputIsTicked(Data, FrameTime));
+	THEN(PressingKeyTriggersStarted(Data, TestAction));
+
+	for (int32 Frame = 1; Frame < HoldFrames - 1; ++Frame)
+	{
+		WHEN(InputIsTicked(Data, FrameTime));
+		THEN(HoldingKeyTriggersOngoing(Data, TestAction));
+	}
+
+	WHEN(InputIsTicked(Data, FrameTime));
+	THEN(HoldingKeyTriggersAction(Data, TestAction)); // fires again on re-press
+
+	// bIsOneShot: returns None immediately after firing -> Triggered->None=Completed.
+	WHEN(InputIsTicked(Data, FrameTime));
+	THEN(HoldingKeyTriggersCompleted(Data, TestAction));
+
+	return true;
+}
+
+// ============================================================
+// HoldAndRelease trigger (integration - replaces the disabled unit test)
+// ============================================================
+
+IMPLEMENT_SIMPLE_AUTOMATION_TEST(FInputIntegrationTriggerHoldAndReleaseTest, "Input.Integration.Triggers.HoldAndRelease", BasicTriggerExtTestFlags)
+
+bool FInputIntegrationTriggerHoldAndReleaseTest::RunTest(const FString& Parameters)
+{
+	const float FrameTime = 1.f / 60.f;
+	const int32 HoldFrames = 30;
+
+	UControllablePlayer* DataPtr = nullptr;
+	UInputTriggerHoldAndRelease* HARTrigger = ABasicTriggerExtTest<UInputTriggerHoldAndRelease>(this, DataPtr);
+	check(HARTrigger && DataPtr);
+	UControllablePlayer& Data = *DataPtr;
+
+	HARTrigger->HoldTimeThreshold = FrameTime * HoldFrames;
+
+	// -------------------------------------------------------
+	// Test 1 - Release before threshold -> Canceled (never triggers)
+	// -------------------------------------------------------
+
+	WHEN(AKeyIsActuated(Data, TestKey));
+	AND(InputIsTicked(Data, FrameTime));
+	THEN(PressingKeyTriggersStarted(Data, TestAction));
+
+	// Release before threshold
+	WHEN(AKeyIsReleased(Data, TestKey));
+	AND(InputIsTicked(Data, FrameTime));
+	THEN(ReleasingKeyTriggersCanceled(Data, TestAction));
+
+	// No further events
+	WHEN(InputIsTicked(Data, FrameTime));
+	THEN(ReleasingKeyDoesNotTrigger(Data, TestAction));
+
+	// -------------------------------------------------------
+	// Test 2 - Hold to threshold then release -> Triggered
+	// -------------------------------------------------------
+	InputIsTicked(Data, FrameTime); // clear state
+
+	WHEN(AKeyIsActuated(Data, TestKey));
+	AND(InputIsTicked(Data, FrameTime));
+	THEN(PressingKeyTriggersStarted(Data, TestAction));
+
+	for (int32 Frame = 1; Frame < HoldFrames - 1; ++Frame)
+	{
+		WHEN(InputIsTicked(Data, FrameTime));
+		THEN(HoldingKeyTriggersOngoing(Data, TestAction));
+	}
+
+	// Hold one more frame (reaches threshold)
+	WHEN(InputIsTicked(Data, FrameTime));
+	THEN(HoldingKeyTriggersOngoing(Data, TestAction)); // still ongoing (not released yet)
+
+	// Release on the threshold frame triggers
+	WHEN(AKeyIsReleased(Data, TestKey));
+	AND(InputIsTicked(Data, FrameTime));
+	THEN(ReleasingKeyTriggersAction(Data, TestAction));
+
+	// Completed on next tick
+	WHEN(InputIsTicked(Data, FrameTime));
+	THEN(ReleasingKeyTriggersCompleted(Data, TestAction));
+
+	WHEN(InputIsTicked(Data, FrameTime));
+	THEN(ReleasingKeyDoesNotTrigger(Data, TestAction));
+
+	// -------------------------------------------------------
+	// Test 3 - Hold past threshold then release -> also Triggered
+	// -------------------------------------------------------
+	InputIsTicked(Data, FrameTime); // clear state
+
+	WHEN(AKeyIsActuated(Data, TestKey));
+	AND(InputIsTicked(Data, FrameTime));
+	THEN(PressingKeyTriggersStarted(Data, TestAction));
+
+	// Hold well past threshold
+	for (int32 Frame = 1; Frame < HoldFrames + 10; ++Frame)
+	{
+		WHEN(InputIsTicked(Data, FrameTime));
+		THEN(HoldingKeyTriggersOngoing(Data, TestAction));
+	}
+
+	// Release past threshold -> still triggers
+	WHEN(AKeyIsReleased(Data, TestKey));
+	AND(InputIsTicked(Data, FrameTime));
+	THEN(ReleasingKeyTriggersAction(Data, TestAction));
+
+	return true;
+}
+
+// ============================================================
+// Tap trigger (integration - replaces the disabled unit test)
+// ============================================================
+
+IMPLEMENT_SIMPLE_AUTOMATION_TEST(FInputIntegrationTriggerTapTest, "Input.Integration.Triggers.Tap", BasicTriggerExtTestFlags)
+
+bool FInputIntegrationTriggerTapTest::RunTest(const FString& Parameters)
+{
+	const float FrameTime = 1.f / 60.f;
+	// Use a 10-frame tap window (consistent with the disabled unit test)
+	const int32 TapWindowFrames = 10;
+	const float TapThreshold = FrameTime * TapWindowFrames;
+
+	UControllablePlayer* DataPtr = nullptr;
+	UInputTriggerTap* TapTrigger = ABasicTriggerExtTest<UInputTriggerTap>(this, DataPtr);
+	check(TapTrigger && DataPtr);
+	UControllablePlayer& Data = *DataPtr;
+
+	TapTrigger->TapReleaseTimeThreshold = TapThreshold;
+
+	// -------------------------------------------------------
+	// Test 1 - Quick release (frame 1, well within window) -> Triggered
+	// -------------------------------------------------------
+	WHEN(AKeyIsActuated(Data, TestKey));
+	AND(InputIsTicked(Data, FrameTime));
+	THEN(PressingKeyTriggersStarted(Data, TestAction));
+
+	WHEN(AKeyIsReleased(Data, TestKey));
+	AND(InputIsTicked(Data, FrameTime));
+	THEN(ReleasingKeyTriggersAction(Data, TestAction));
+
+	WHEN(InputIsTicked(Data, FrameTime));
+	THEN(ReleasingKeyTriggersCompleted(Data, TestAction));
+
+	InputIsTicked(Data, FrameTime); // clear state
+
+	// -------------------------------------------------------
+	// Test 2 - Release on the last valid frame (frame TapWindowFrames-1, just before threshold)
+	//          HeldDuration < TapThreshold on release -> Triggered
+	// -------------------------------------------------------
+	WHEN(AKeyIsActuated(Data, TestKey));
+	AND(InputIsTicked(Data, FrameTime));
+	THEN(PressingKeyTriggersStarted(Data, TestAction));
+
+	for (int32 Frame = 1; Frame < TapWindowFrames - 1; ++Frame)
+	{
+		WHEN(InputIsTicked(Data, FrameTime));
+		THEN(HoldingKeyTriggersOngoing(Data, TestAction));
+	}
+
+	WHEN(AKeyIsReleased(Data, TestKey));
+	AND(InputIsTicked(Data, FrameTime));
+	THEN(ReleasingKeyTriggersAction(Data, TestAction));
+
+	InputIsTicked(Data, FrameTime); // clear state
+
+	// -------------------------------------------------------
+	// Test 3 - Hold past the tap window -> tap Triggered event never fires.
+	//          While held: HeldDuration < TapThreshold -> Ongoing
+	//          Once HeldDuration >= TapThreshold: trigger returns None, action fires Canceled (window missed)
+	//          Releasing after the window: no Triggered event (tap was not registered)
+	// -------------------------------------------------------
+	WHEN(AKeyIsActuated(Data, TestKey));
+	AND(InputIsTicked(Data, FrameTime));
+	THEN(PressingKeyTriggersStarted(Data, TestAction));
+
+	// Hold through the tap window: Ongoing on each frame up to (not including) the threshold tick.
+	// At tick TapWindowFrames, HeldDuration >= TapThreshold fires None (window expired).
+	// Stop one tick early so the final loop tick is still Ongoing.
+	for (int32 Frame = 1; Frame < TapWindowFrames - 1; ++Frame)
+	{
+		WHEN(InputIsTicked(Data, FrameTime));
+		THEN(HoldingKeyTriggersOngoing(Data, TestAction));
+		AND(TestFalse(TEXT("Tap: Triggered should not fire while holding within window"), FInputTestHelper::TestTriggered(Data, TestAction)));
+	}
+
+	// One more tick: HeldDuration now >= TapThreshold, trigger returns None -> Canceled fires (window expired)
+	WHEN(InputIsTicked(Data, FrameTime));
+	TestFalse(TEXT("Tap: Triggered must not fire after window expires"), FInputTestHelper::TestTriggered(Data, TestAction));
+
+	// Releasing after the tap window: Triggered still does not fire
+	WHEN(AKeyIsReleased(Data, TestKey));
+	AND(InputIsTicked(Data, FrameTime));
+	TestFalse(TEXT("Tap: Triggered must not fire on release after window"), FInputTestHelper::TestTriggered(Data, TestAction));
+
+	return true;
+}
+
+// ============================================================
+// ChordBlocker trigger
+// ============================================================
+
+IMPLEMENT_SIMPLE_AUTOMATION_TEST(FInputTriggerChordBlockerTest, "Input.Triggers.Chords.Blocker", BasicTriggerExtTestFlags)
+
+bool FInputTriggerChordBlockerTest::RunTest(const FString& Parameters)
+{
+	// Setup:
+	//   MainAction  mapped to TestKey  with a ChordBlocker referencing ChordingAction.
+	//   ChordingAction mapped to TestKey2 (no trigger - fires while held).
+	//
+	// Expected behavior:
+	//   TestKey alone           -> MainAction fires (ChordingAction inactive, blocker not blocking)
+	//   TestKey2 alone          -> ChordingAction fires, MainAction unaffected (not pressed)
+	//   TestKey + TestKey2      -> ChordingAction fires, MainAction is BLOCKED
+	//   TestKey2 released       -> MainAction fires again (blocker no longer blocking)
+
+	UWorld* World =
+	GIVEN(AnEmptyWorld());
+
+	UControllablePlayer& Data =
+	AND(AControllablePlayer(World));
+	TestTrue(TEXT("Controllable Player is valid"), Data.IsValid());
+
+	AND(AnInputContextIsAppliedToAPlayer(Data, TestContext, 0));
+	AND(AnInputAction(Data, TestAction, EInputActionValueType::Boolean));   // MainAction
+	AND(AnInputAction(Data, TestAction2, EInputActionValueType::Boolean));  // ChordingAction
+
+	AND(AnActionIsMappedToAKey(Data, TestContext, TestAction, TestKey));
+	AND(AnActionIsMappedToAKey(Data, TestContext, TestAction2, TestKey2));
+
+	// Add a Down (Explicit) trigger so the action is gated on key actuation.
+	// Without an Explicit trigger, a Blocker-only action fires Triggered unconditionally
+	// (bFoundExplicit=false, bAllImplicitsTriggered=true -> bTriggered=true always when not blocked).
+	AND(ATriggerIsAppliedToAnAction(Data, NewObject<UInputTriggerDown>(), TestAction));
+
+	// Apply ChordBlocker(ChordingAction) to MainAction (added after Down so its pointer stays valid)
+	UInputTriggerChordBlocker* Blocker =
+	AND(Cast<UInputTriggerChordBlocker>(ATriggerIsAppliedToAnAction(Data, NewObject<UInputTriggerChordBlocker>(), TestAction)));
+	TestNotNull(TEXT("ChordBlocker created"), Blocker);
+	Blocker->ChordAction = FInputTestHelper::FindAction(Data, TestAction2);
+	TestNotNull(TEXT("ChordAction assigned to blocker"), Blocker->ChordAction.Get());
+
+	// Test 1 - TestKey alone: MainAction fires, ChordingAction not active so blocker does not block
+	WHEN(AKeyIsActuated(Data, TestKey));
+	AND(InputIsTicked(Data));
+	THEN(PressingKeyTriggersAction(Data, TestAction));
+	AND(FInputTestHelper::TestNoTrigger(Data, TestAction2)); // ChordingAction not pressed
+
+	// Test 2 - Press TestKey2 while TestKey is still held:
+	//          ChordingAction activates -> blocker becomes active -> MainAction blocked
+	WHEN(AKeyIsActuated(Data, TestKey2));
+	AND(InputIsTicked(Data));
+	THEN(HoldingKeyTriggersAction(Data, TestAction2));       // ChordingAction fires
+	AND(FInputTestHelper::TestNoTrigger(Data, TestAction));  // MainAction is blocked
+
+	// Test 3 - Hold both keys: MainAction continues to be blocked
+	WHEN(InputIsTicked(Data));
+	THEN(HoldingKeyTriggersAction(Data, TestAction2));
+	AND(FInputTestHelper::TestNoTrigger(Data, TestAction));
+
+	// Test 4 - Release TestKey2: ChordingAction stops -> blocker deactivates -> MainAction fires again
+	WHEN(AKeyIsReleased(Data, TestKey2));
+	AND(InputIsTicked(Data));
+	THEN(HoldingKeyTriggersAction(Data, TestAction));        // MainAction unblocked
+
+	// Test 5 - Release TestKey: MainAction completes
+	WHEN(AKeyIsReleased(Data, TestKey));
+	AND(InputIsTicked(Data));
+	THEN(ReleasingKeyTriggersCompleted(Data, TestAction));
+
+	return true;
+}
