[Enhanced Input] Rename the InputEditorModule.cpp to EnhancedInputEditorModule.cpp, as well as the InputCustomizations.h/cpp files to make them easier to find.

#Synced-CL 52004614

diff --git a/Engine/Plugins/EnhancedInput/Source/InputEditor/Private/EnhancedInputCustomizations.cpp b/Engine/Plugins/EnhancedInput/Source/InputEditor/Private/EnhancedInputCustomizations.cpp
new file mode 100644
--- /dev/null
+++ b/Engine/Plugins/EnhancedInput/Source/InputEditor/Private/EnhancedInputCustomizations.cpp
@@ -0,0 +1,705 @@
+// Copyright Epic Games, Inc. All Rights Reserved.
+
+#include "EnhancedInputCustomizations.h"
+
+#include "ActionMappingDetails.h"
+#include "AssetRegistry/ARFilter.h"
+#include "DetailLayoutBuilder.h"
+#include "Engine/Blueprint.h"
+#include "IDetailChildrenBuilder.h"
+#include "InputMappingContext.h"
+#include "KeyStructCustomization.h"
+#include "PropertyCustomizationHelpers.h"
+#include "EnhancedInputDeveloperSettings.h"
+#include "AssetRegistry/AssetRegistryModule.h"
+#include "Blueprint/BlueprintSupport.h"
+#include "Misc/PackageName.h"
+#include "UObject/UObjectIterator.h"
+#include "Modules/ModuleManager.h"
+#include "Widgets/Text/STextBlock.h"
+#include "EnhancedInputPlayerMappableNameValidator.h"
+#include "PlayerMappableKeySettings.h"
+#include "Widgets/Input/SEditableTextBox.h"
+#include "EnhancedInputPlayerMappableNameValidator.h"
+#include "InputEditorModule.h"
+
+#define LOCTEXT_NAMESPACE "InputCustomization"
+
+
+//////////////////////////////////////////////////////////
+// FEnhancedActionMappingCustomization
+
+const FEnhancedActionKeyMapping* GetActionKeyMapping(TSharedPtr<IPropertyHandle> MappingPropertyHandle)
+{
+	if (MappingPropertyHandle->IsValidHandle())
+	{
+		void* Data = nullptr;
+		if (MappingPropertyHandle->GetValueData(Data) != FPropertyAccess::Fail)
+		{
+			return static_cast<FEnhancedActionKeyMapping*>(Data);
+		}
+	}
+
+	return nullptr;
+}
+
+void FEnhancedActionMappingCustomization::CustomizeHeader(TSharedRef<IPropertyHandle> PropertyHandle, FDetailWidgetRow& HeaderRow, IPropertyTypeCustomizationUtils& CustomizationUtils)
+{
+	MappingPropertyHandle = PropertyHandle;
+
+	// Grab the FKey property
+	TSharedPtr<IPropertyHandle> KeyHandle = PropertyHandle->GetChildHandle(GET_MEMBER_NAME_CHECKED(FEnhancedActionKeyMapping, Key));
+	TSharedPtr<IPropertyHandle> TriggersHandle = PropertyHandle->GetChildHandle(GET_MEMBER_NAME_CHECKED(FEnhancedActionKeyMapping, Triggers));
+
+	TriggersHandle->SetOnPropertyValueChanged(FSimpleDelegate::CreateSP(this, &FEnhancedActionMappingCustomization::OnTriggersChanged));
+	TriggersHandle->SetOnChildPropertyValueChanged(FSimpleDelegate::CreateSP(this, &FEnhancedActionMappingCustomization::OnTriggersChanged));
+	TSharedRef<SWidget> RemoveButton = PropertyCustomizationHelpers::MakeDeleteButton(FSimpleDelegate::CreateSP(this, &FEnhancedActionMappingCustomization::RemoveMappingButton_OnClick),
+		LOCTEXT("RemoveMappingToolTip", "Remove Mapping"));
+
+	// Create  a new instance of the key customization.
+	KeyStructInstance = FKeyStructCustomization::MakeInstance();
+
+	// TODO: Use FDetailArrayBuilder?
+
+	PRAGMA_DISABLE_DEPRECATION_WARNINGS
+	const bool bContainsComboTrigger = DoesTriggerArrayContainCombo();
+	PRAGMA_ENABLE_DEPRECATION_WARNINGS
+	// Pass our header row into the key struct customizeheader method so it populates our row with the key struct header
+	KeyStructCustomization = StaticCastSharedPtr<FKeyStructCustomization>(KeyStructInstance);
+	KeyStructCustomization->SetDefaultKeyName("ComboKey");
+	KeyStructCustomization->SetDisabledKeySelectorToolTip(LOCTEXT("DisabledSelectorComboToolTip", "A Combo Trigger isn't triggered through a key so Key Selection has been disabled."));
+	KeyStructCustomization->SetEnableKeySelector(!bContainsComboTrigger);
+	KeyStructCustomization->CustomizeHeaderOnlyWithButton(KeyHandle.ToSharedRef(), HeaderRow, CustomizationUtils, RemoveButton);
+}
+
+void FEnhancedActionMappingCustomization::AddInputActionProperties(TSharedRef<IPropertyHandle> PropertyHandle, IDetailChildrenBuilder& ChildBuilder)
+{
+    if (const FEnhancedActionKeyMapping* ActionKeyMapping = GetActionKeyMapping(PropertyHandle))
+    {
+    	// Make sure ActionKeyMapping action is valid. If triggers and modifiers are empty we can just back out - nothing to add
+    	if (ActionKeyMapping->Action)
+    	{
+    		// Convert InputAction to non const so we can properly use it in the UObject Array for AddExternalObjectProperty
+    		const UInputAction* InputActionPtr = ActionKeyMapping->Action;
+    		UInputAction* InputAction = const_cast<UInputAction*>(InputActionPtr);
+    		TArray<UObject*> ActionsAsUObjects { InputAction };
+
+    		if (IDetailPropertyRow* InputActionTriggersRow = ChildBuilder.AddExternalObjectProperty(ActionsAsUObjects, GET_MEMBER_NAME_CHECKED(UInputAction, Triggers)))
+    		{
+    			InputActionTriggersPropertyRow = InputActionTriggersRow;
+    			InputActionTriggersRow->DisplayName(LOCTEXT("InputActionTriggersDisplayName", "Triggers From Input Action"));
+    			InputActionTriggersRow->ToolTip(FText::Format(LOCTEXT("InputActionTriggersToolTip", "Triggers from the {0} Input Action"), FText::FromName(ActionKeyMapping->Action.GetFName())));
+    			InputActionTriggersRow->IsEnabled(false);
+    			InputAction->OnTriggersChanged.AddSP(this, &FEnhancedActionMappingCustomization::OnInputActionTriggersChanged);
+    			
+    			// hide it if the trigger array is empty
+    			if (ActionKeyMapping->Action->Triggers.IsEmpty())
+    			{
+					InputActionTriggersRow->Visibility(EVisibility::Hidden);
+    			}
+    			else
+    			{
+    				InputActionTriggersRow->Visibility(EVisibility::Visible);
+    			}
+    		}
+    		if (IDetailPropertyRow* InputActionModifiersRow = ChildBuilder.AddExternalObjectProperty(ActionsAsUObjects, GET_MEMBER_NAME_CHECKED(UInputAction, Modifiers)))
+    		{
+    			InputActionModifiersPropertyRow = InputActionModifiersRow;
+    			InputActionModifiersRow->DisplayName(LOCTEXT("InputActionModifiersDisplayName", "Modifiers From Input Action"));
+    			InputActionModifiersRow->ToolTip(FText::Format(LOCTEXT("InputActionModifiersToolTip", "Modifiers from the {0} Input Action"), FText::FromName(ActionKeyMapping->Action.GetFName())));
+    			InputActionModifiersRow->IsEnabled(false);
+    			InputAction->OnModifiersChanged.AddSP(this, &FEnhancedActionMappingCustomization::OnInputActionModifiersChanged);
+    			
+    			// hide it if the modifier array is empty
+    			if (ActionKeyMapping->Action->Modifiers.IsEmpty())
+    			{
+    				InputActionModifiersRow->Visibility(EVisibility::Hidden);
+    			}
+    			else
+    			{
+    				InputActionModifiersRow->Visibility(EVisibility::Visible);	
+    			}
+    		}
+    	}
+    }
+}
+
+void FEnhancedActionMappingCustomization::CustomizeChildren(TSharedRef<IPropertyHandle> PropertyHandle, IDetailChildrenBuilder& ChildBuilder, IPropertyTypeCustomizationUtils& CustomizationUtils)
+{
+	TSharedPtr<IPropertyHandle> TriggersHandle = PropertyHandle->GetChildHandle(GET_MEMBER_NAME_CHECKED(FEnhancedActionKeyMapping, Triggers));
+	TSharedPtr<IPropertyHandle> ModifiersHandle = PropertyHandle->GetChildHandle(GET_MEMBER_NAME_CHECKED(FEnhancedActionKeyMapping, Modifiers));
+	TSharedPtr<IPropertyHandle> SettingBehavior = PropertyHandle->GetChildHandle(GET_MEMBER_NAME_CHECKED(FEnhancedActionKeyMapping, SettingBehavior));
+	TSharedPtr<IPropertyHandle> PlayerMappableKeySettings = PropertyHandle->GetChildHandle(GET_MEMBER_NAME_CHECKED(FEnhancedActionKeyMapping, PlayerMappableKeySettings));
+
+	// TODO: ResetToDefault needs to be disabled for arrays
+	ChildBuilder.AddProperty(TriggersHandle.ToSharedRef());
+	ChildBuilder.AddProperty(ModifiersHandle.ToSharedRef());
+	AddInputActionProperties(PropertyHandle, ChildBuilder);
+	ChildBuilder.AddProperty(SettingBehavior.ToSharedRef());
+	ChildBuilder.AddProperty(PlayerMappableKeySettings.ToSharedRef());
+}
+
+void FEnhancedActionMappingCustomization::RemoveMappingButton_OnClick() const
+{
+	if (MappingPropertyHandle->IsValidHandle())
+	{
+		const TSharedPtr<IPropertyHandle> ParentHandle = MappingPropertyHandle->GetParentHandle();
+		const TSharedPtr<IPropertyHandleArray> ParentArrayHandle = ParentHandle->AsArray();
+
+		ParentArrayHandle->DeleteItem(MappingPropertyHandle->GetIndexInArray());
+	}
+}
+
+void FEnhancedActionMappingCustomization::OnTriggersChanged() const
+{
+	PRAGMA_DISABLE_DEPRECATION_WARNINGS
+	const bool bContainsComboTrigger = DoesTriggerArrayContainCombo();
+	PRAGMA_ENABLE_DEPRECATION_WARNINGS
+	
+	// if it is currently disabled and doesn't contain a combo now we should set the key to something other than ComboKey
+	if (!KeyStructCustomization->GetEnableKeySelector() && !bContainsComboTrigger)
+	{
+		// setting KeySelector to none key
+		KeyStructCustomization->SetKey(TEXT("None"));
+	}
+	
+    KeyStructCustomization->SetEnableKeySelector(!bContainsComboTrigger);
+	// updating the default key when the KeySelector is disabled
+	KeyStructCustomization->SetDefaultKeyName(bContainsComboTrigger ? TEXT("ComboKey") : TEXT("None"));
+}
+
+void FEnhancedActionMappingCustomization::OnInputActionTriggersChanged() const
+{
+	if (const FEnhancedActionKeyMapping* ActionKeyMapping = GetActionKeyMapping(MappingPropertyHandle))
+	{
+		// Make sure ActionKeyMapping action and the property row are valid
+		if (ActionKeyMapping->Action && InputActionTriggersPropertyRow)
+		{
+			// if so we want to hide the row or show it based off contents of the array (whether it's empty or not)
+			if (!ActionKeyMapping->Action->Triggers.IsEmpty())
+			{
+				InputActionTriggersPropertyRow->Visibility(EVisibility::Visible);
+			}
+			else
+			{
+				InputActionTriggersPropertyRow->Visibility(EVisibility::Hidden);
+			}
+		}
+	}
+}
+
+void FEnhancedActionMappingCustomization::OnInputActionModifiersChanged() const
+{
+    if (const FEnhancedActionKeyMapping* ActionKeyMapping = GetActionKeyMapping(MappingPropertyHandle))
+    {
+    	// Make sure ActionKeyMapping action and the property row are valid
+    	if (ActionKeyMapping->Action && InputActionModifiersPropertyRow)
+    	{
+    		// if so we want to hide the row or show it based off contents of the array (whether it's empty or not)
+    		if (!ActionKeyMapping->Action->Modifiers.IsEmpty())
+    		{
+    			InputActionModifiersPropertyRow->Visibility(EVisibility::Visible);
+    		}
+    		else
+    		{
+    			InputActionModifiersPropertyRow->Visibility(EVisibility::Hidden);
+    		}
+    	}
+    }
+}
+
+PRAGMA_DISABLE_DEPRECATION_WARNINGS
+bool FEnhancedActionMappingCustomization::DoesTriggerArrayContainCombo() const
+{
+	if (const FEnhancedActionKeyMapping* ActionKeyMapping = GetActionKeyMapping(MappingPropertyHandle))
+	{
+		// checking context triggers for combo triggers
+		for (const TObjectPtr<UInputTrigger>& Trigger : ActionKeyMapping->Triggers)
+		{
+			if (Trigger.IsA(UInputTriggerCombo::StaticClass()))
+			{
+				return true;
+			}
+		}
+		// checking input action triggers for combo triggers
+		if (ActionKeyMapping->Action)
+		{
+			for (const TObjectPtr<UInputTrigger>& Trigger : ActionKeyMapping->Action->Triggers)
+			{
+				if (Trigger.IsA(UInputTriggerCombo::StaticClass()))
+				{
+					return true;
+				}
+			}
+		}
+	}
+	
+	return false;
+}
+PRAGMA_ENABLE_DEPRECATION_WARNINGS
+
+//////////////////////////////////////////////////////////
+// FEnhancedInputDeveloperSettingsCustomization
+
+TSet<FName> FEnhancedInputDeveloperSettingsCustomization::ExcludedAssetNames;
+
+TSharedRef<IDetailCustomization> FEnhancedInputDeveloperSettingsCustomization::MakeInstance()
+{
+	LLM_SCOPE_BYNAME(TEXT("EnhancedInput/Editor"));
+	return MakeShareable(new FEnhancedInputDeveloperSettingsCustomization());
+}
+
+FEnhancedInputDeveloperSettingsCustomization::~FEnhancedInputDeveloperSettingsCustomization()
+{
+	// Unregister settings panel listeners
+	if (FAssetRegistryModule* AssetRegistryModule = FModuleManager::GetModulePtr<FAssetRegistryModule>("AssetRegistry"))
+	{
+		IAssetRegistry* AssetRegistry = AssetRegistryModule->TryGet();
+		if (AssetRegistry)
+		{
+			AssetRegistry->OnAssetAdded().RemoveAll(this);
+			AssetRegistry->OnAssetRemoved().RemoveAll(this);
+			AssetRegistry->OnAssetRenamed().RemoveAll(this);
+		}
+	}
+	
+	CachedDetailBuilder = nullptr;
+}
+
+void FEnhancedInputDeveloperSettingsCustomization::CustomizeDetails(IDetailLayoutBuilder& DetailBuilder)
+{
+	// Note: The details view for the UEnhancedInputDeveloperSettings object will be on by default
+	// This 'EditCategory' call will ensure that the developer settings are displayed on top of the trigger/modifier default values
+	DetailBuilder.EditCategory(UEnhancedInputDeveloperSettings::StaticClass()->GetFName(), FText::GetEmpty(), ECategoryPriority::Important);
+
+	static const FName TriggerCategoryName = TEXT("Trigger Default Values");
+	static const FName ModifierCategoryName = TEXT("Modifier Default Values");
+
+	FAssetRegistryModule& AssetRegistryModule = FModuleManager::LoadModuleChecked<FAssetRegistryModule>(FName("AssetRegistry"));
+	IAssetRegistry& AssetRegistry = AssetRegistryModule.Get();
+
+	TArray<UObject*> ModifierCDOs;
+	TArray<UObject*> TriggerCDOs;
+	// Search BPs via asset registry
+	{
+		FARFilter Filter;
+		Filter.ClassPaths.Add(UBlueprint::StaticClass()->GetClassPathName());
+		Filter.bRecursiveClasses = true;
+		TArray<FAssetData> BlueprintAssetData;
+		AssetRegistry.GetAssets(Filter, BlueprintAssetData);
+
+		FName NativeParentClass(TEXT("NativeParentClass"));
+		for (FAssetData& Asset : BlueprintAssetData)
+		{
+			FAssetDataTagMapSharedView::FFindTagResult Result = Asset.TagsAndValues.FindTag(NativeParentClass);
+			if (Result.IsSet() && !ExcludedAssetNames.Contains(Asset.AssetName))
+			{
+				const FString ClassObjectPath = FPackageName::ExportTextPathToObjectPath(Result.GetValue());
+				if (UClass* ParentClass = FindObjectSafe<UClass>(nullptr, *ClassObjectPath, EFindObjectFlags::ExactClass))
+				{
+					// TODO: Forcibly loading these assets could cause problems on projects with a large number of them.
+					if(ParentClass->IsChildOf(UInputModifier::StaticClass()))
+					{
+						UBlueprint* BP = CastChecked<UBlueprint>(Asset.GetAsset());
+						ModifierCDOs.AddUnique(BP->GeneratedClass->GetDefaultObject());
+					}
+					if (ParentClass->IsChildOf(UInputTrigger::StaticClass()))
+					{
+						UBlueprint* BP = CastChecked<UBlueprint>(Asset.GetAsset());
+						TriggerCDOs.AddUnique(BP->GeneratedClass->GetDefaultObject());
+					}
+				}
+			}
+		}
+	}
+
+	GatherNativeClassDetailsCDOs(UInputModifier::StaticClass(), ModifierCDOs);
+	GatherNativeClassDetailsCDOs(UInputTrigger::StaticClass(), TriggerCDOs);
+	ExcludedAssetNames.Reset();
+	
+	// Add The modifier/trigger defaults that are generated via CDO to the details builder
+	CustomizeCDOValues(DetailBuilder, ModifierCategoryName, ModifierCDOs);
+	CustomizeCDOValues(DetailBuilder, TriggerCategoryName, TriggerCDOs);
+	
+	// Support for updating blueprint based triggers and modifiers in the settings panel
+	if (!AssetRegistry.OnAssetAdded().IsBoundToObject(this))
+	{
+		AssetRegistry.OnAssetAdded().AddRaw(this, &FEnhancedInputDeveloperSettingsCustomization::OnAssetAdded);
+		AssetRegistry.OnAssetRemoved().AddRaw(this, &FEnhancedInputDeveloperSettingsCustomization::OnAssetRemoved);
+		AssetRegistry.OnAssetRenamed().AddRaw(this, &FEnhancedInputDeveloperSettingsCustomization::OnAssetRenamed);
+	}
+}
+
+void FEnhancedInputDeveloperSettingsCustomization::CustomizeCDOValues(IDetailLayoutBuilder& DetailBuilder, const FName CategoryName, const TArray<UObject*>& ObjectsToCustomize)
+{
+	IDetailCategoryBuilder& CategoryBuilder = DetailBuilder.EditCategory(CategoryName);
+
+	// All of the Objects in this array are the CDO, i.e. the class name starts with "Default__"
+	for (UObject* CDO : ObjectsToCustomize)
+	{
+		if (!ensure(CDO && CDO->IsTemplate()))
+		{
+			continue;
+		}
+
+		// Add the CDO as an external object reference to this customization
+		IDetailPropertyRow* Row = CategoryBuilder.AddExternalObjects({ CDO }, EPropertyLocation::Default, FAddPropertyParams().UniqueId(CDO->GetClass()->GetFName()));
+
+		if (Row)
+		{
+			// We need to add a custom "Name" widget here, otherwise all the categories will just say "Object"
+			Row->CustomWidget()
+			.NameContent()
+			[
+				SNew(STextBlock)
+					.Text(CDO->GetClass()->GetDisplayNameText())
+			];	
+		}
+	}
+}
+
+void FEnhancedInputDeveloperSettingsCustomization::CustomizeDetails(const TSharedPtr<IDetailLayoutBuilder>& DetailBuilder)
+{
+	CachedDetailBuilder = DetailBuilder;
+	CustomizeDetails(*DetailBuilder);
+}
+
+void FEnhancedInputDeveloperSettingsCustomization::GatherNativeClassDetailsCDOs(UClass* Class, TArray<UObject*>& CDOs)
+{
+	// Search native classes
+	for (TObjectIterator<UClass> ClassIt; ClassIt; ++ClassIt)
+	{
+		if (!ClassIt->IsNative() || !ClassIt->IsChildOf(Class))
+		{
+			continue;
+		}
+
+		// Ignore abstract, hidedropdown, and deprecated.
+		if (ClassIt->HasAnyClassFlags(CLASS_Abstract | CLASS_HideDropDown | CLASS_Deprecated | CLASS_NewerVersionExists))
+		{
+			continue;
+		}
+
+		CDOs.AddUnique(ClassIt->GetDefaultObject());
+	}
+
+	// Strip objects with no config stored properties
+	CDOs.RemoveAll([Class](UObject* Object) {
+		UClass* ObjectClass = Object->GetClass();
+		if (ObjectClass->GetMetaData(TEXT("NotInputConfigurable")).ToBool())
+		{
+			return true;
+		}
+		while (ObjectClass)
+		{
+			for (FProperty* Property : TFieldRange<FProperty>(ObjectClass, EFieldIteratorFlags::ExcludeSuper, EFieldIteratorFlags::ExcludeDeprecated))
+			{
+				if (Property->HasAnyPropertyFlags(CPF_Config))
+				{
+					return false;
+				}
+			}
+
+			// Stop searching at the base type. We don't care about configurable properties lower than that.
+			ObjectClass = ObjectClass != Class ? ObjectClass->GetSuperClass() : nullptr;
+		}
+		return true;
+	});
+}
+
+bool FEnhancedInputDeveloperSettingsCustomization::DoesClassHaveSubtypes(UClass* Class)
+{
+	// Search native classes
+	for (TObjectIterator<UClass> ClassIt; ClassIt; ++ClassIt)
+	{
+		// Make sure it's a native class and a child of Class passed in
+		if (ClassIt->IsNative() && ClassIt->IsChildOf(Class))
+		{
+			// Make sure it doesn't have any flags that would disqualify it from being used
+			if (!ClassIt->HasAnyClassFlags(CLASS_Abstract | CLASS_HideDropDown | CLASS_Deprecated | CLASS_NewerVersionExists))
+			{
+				UObject* CDO = ClassIt->GetDefaultObject();
+		
+				// Make Sure it isn't the Class itself
+				if (CDO && CDO->GetClass() != Class)
+				{
+					return true;
+				}
+			}
+		}
+	}
+
+	// Search BPs via asset registry
+	FAssetRegistryModule& AssetRegistryModule = FModuleManager::LoadModuleChecked<FAssetRegistryModule>(FName("AssetRegistry"));
+	IAssetRegistry& AssetRegistry = AssetRegistryModule.Get();
+
+	FARFilter Filter;
+	Filter.ClassPaths.Add(UBlueprint::StaticClass()->GetClassPathName());
+	Filter.bRecursiveClasses = true;
+	TArray<FAssetData> BlueprintAssetData;
+	AssetRegistry.GetAssets(Filter, BlueprintAssetData);
+
+	for (FAssetData& Asset : BlueprintAssetData)
+	{
+		FAssetDataTagMapSharedView::FFindTagResult Result = Asset.TagsAndValues.FindTag(TEXT("NativeParentClass"));
+		if (Result.IsSet() && !ExcludedAssetNames.Contains(Asset.AssetName))
+		{
+			const FString ClassObjectPath = FPackageName::ExportTextPathToObjectPath(Result.GetValue());
+			if (UClass* ParentClass = FindObjectSafe<UClass>(nullptr, *ClassObjectPath, EFindObjectFlags::ExactClass))
+			{
+				if (ParentClass->IsChildOf(Class))
+				{
+					return true;
+				}
+			}
+		}
+	}
+
+	return false;
+}
+
+void FEnhancedInputDeveloperSettingsCustomization::RebuildDetailsViewForAsset(const FAssetData& AssetData, const bool bIsAssetBeingRemoved)
+{
+	LLM_SCOPE_BYNAME(TEXT("EnhancedInput/Editor"));
+	
+	// If the asset was a blueprint...
+	if (AssetData.AssetClassPath == UBlueprint::StaticClass()->GetClassPathName())
+	{
+		// With a native parent class...
+		FAssetDataTagMapSharedView::FFindTagResult Result = AssetData.TagsAndValues.FindTag(FBlueprintTags::NativeParentClassPath);
+		if (Result.IsSet())
+		{
+			// And the base class is UInputModifier or UInputTrigger, then we should rebuild the details
+			const FString ClassObjectPath = FPackageName::ExportTextPathToObjectPath(Result.GetValue());
+			if (UClass* ParentClass = FindObjectSafe<UClass>(nullptr, *ClassObjectPath, EFindObjectFlags::ExactClass))
+			{
+				if (ParentClass == UInputModifier::StaticClass() || ParentClass == UInputTrigger::StaticClass())
+				{
+					if (bIsAssetBeingRemoved)
+					{
+						ExcludedAssetNames.Add(AssetData.AssetName);
+					}
+					
+					if (IDetailLayoutBuilder* DetailBuilder = CachedDetailBuilder.Pin().Get())
+					{
+						DetailBuilder->ForceRefreshDetails();
+					}
+				}
+			}
+		}
+	}
+}
+
+/////////////////////////////////////////////////////////////////
+// FPlayerMappableKeyChildSettingsCustomization
+
+namespace UE::EnhancedInput
+{
+	/** Flag to enable or disable name validation in the editor for UPlayerMappableKeySettings assets */
+	static bool bEnableNameValidation = true;
+	static FAutoConsoleVariableRef CVarEnableNameValidation(
+		TEXT("EnhancedInput.bEnableNameValidation"),
+		bEnableNameValidation,
+		TEXT("Flag to enable or disable name validation in the editor for UPlayerMappableKeySettings assets"),
+		ECVF_Default);
+}
+
+TSharedRef<IPropertyTypeCustomization> FPlayerMappableKeyChildSettingsCustomization::MakeInstance()
+{
+	return MakeShareable(new FPlayerMappableKeyChildSettingsCustomization());
+}
+
+void FPlayerMappableKeyChildSettingsCustomization::CustomizeHeader(TSharedRef<IPropertyHandle> PropertyHandle, FDetailWidgetRow& HeaderRow, IPropertyTypeCustomizationUtils& CustomizationUtils)
+{
+	SettingsPropHandle = PropertyHandle;
+	UPlayerMappableKeySettings* Settings = GetSettingsObject();
+	
+	HeaderRow.NameContent()
+	[
+		PropertyHandle->CreatePropertyNameWidget()
+	];
+	
+	HeaderRow.ValueContent()
+	[
+		PropertyHandle->CreatePropertyValueWidget()
+	];	
+}
+
+void FPlayerMappableKeyChildSettingsCustomization::CustomizeChildren(TSharedRef<IPropertyHandle> PropertyHandle, IDetailChildrenBuilder& ChildBuilder, IPropertyTypeCustomizationUtils& CustomizationUtils)
+{				
+	// Customize the name handle to have some validation on what users can type in
+	static const FName MappingPropertyHandleName = GET_MEMBER_NAME_CHECKED(UPlayerMappableKeySettings, Name);
+
+	uint32 NumChildren = 0;
+	TSharedPtr<IPropertyHandle> SettingsHandle;
+	TSharedPtr<IPropertyHandle> TempHandle = PropertyHandle->GetChildHandle(0);
+
+	if (TempHandle.IsValid())
+	{
+		TempHandle->GetNumChildren(NumChildren);
+
+		// For instanced properties there can be multiple layers of "children" until we get to the ones
+		// that we care about. Iterate all the children until we find one with more then one
+		// child, or we run out of children.
+		while (!SettingsHandle.IsValid() || !TempHandle.IsValid())
+		{
+			TempHandle->GetNumChildren(NumChildren);
+
+			if (NumChildren > 1)
+			{
+				SettingsHandle = TempHandle;
+			}
+			else
+			{
+				TempHandle = TempHandle->GetChildHandle(0);
+			}		
+		}	
+	}	
+	
+	if (!SettingsHandle.IsValid())
+	{
+		return;
+	}
+
+	for (uint32 ChildIndex = 0; ChildIndex < NumChildren; ++ChildIndex)
+	{
+		TSharedPtr<IPropertyHandle> ChildHandle = SettingsHandle->GetChildHandle(ChildIndex);
+		if (!ChildHandle.IsValid())
+		{
+			continue;
+		}
+		
+		const FName PropertyName = ChildHandle->GetProperty()->GetFName();
+
+		// Add a customization to handle name validation on the mappable name if it's enabled
+		if (PropertyName == MappingPropertyHandleName && UE::EnhancedInput::bEnableNameValidation)
+		{
+			MappingNamePropHandle = ChildHandle;
+			IDetailPropertyRow& NameRow = ChildBuilder.AddProperty(MappingNamePropHandle.ToSharedRef());       
+            
+            NameRow.CustomWidget()
+            .NameContent()
+            [
+            	MappingNamePropHandle->CreatePropertyNameWidget()
+            ]
+            .ValueContent()
+            [
+            	SNew(SEditableTextBox)
+            		.Text(this, &FPlayerMappableKeyChildSettingsCustomization::OnGetMappingNameText)
+            		.OnTextCommitted(this, &FPlayerMappableKeyChildSettingsCustomization::OnMappingNameTextCommited)
+            		.ToolTipText(this, &FPlayerMappableKeyChildSettingsCustomization::OnGetMappingNameToolTipText)
+            		.OnVerifyTextChanged(this, &FPlayerMappableKeyChildSettingsCustomization::OnVerifyMappingName)
+            		.Font(IDetailLayoutBuilder::GetDetailFont())
+            ];
+		}
+		else
+		{
+			ChildBuilder.AddProperty(ChildHandle.ToSharedRef());
+		}
+	}
+}
+
+FName FPlayerMappableKeyChildSettingsCustomization::GetCurrentName() const
+{
+	if (UPlayerMappableKeySettings* Settings = GetSettingsObject())
+	{
+		return Settings->Name;
+	}
+	ensure(false);
+	return NAME_None;
+}
+
+UPlayerMappableKeySettings* FPlayerMappableKeyChildSettingsCustomization::GetSettingsObject() const
+{
+	if (SettingsPropHandle.IsValid())
+	{
+		UObject* SettingsObject = nullptr;
+		SettingsPropHandle->GetValue(SettingsObject);
+		return Cast<UPlayerMappableKeySettings>(SettingsObject);	
+	}
+
+	return nullptr;
+}
+
+FText FPlayerMappableKeyChildSettingsCustomization::OnGetMappingNameText() const
+{
+	// this is the value in the actual text box
+	if (UPlayerMappableKeySettings* Settings = GetSettingsObject())
+	{
+		return FText::FromName(Settings->Name);
+	}
+	
+	return FText::FromName(NAME_None);
+}
+
+FText FPlayerMappableKeyChildSettingsCustomization::OnGetMappingNameToolTipText() const
+{
+	static const FText Tooltip = LOCTEXT("MappingNameToolTip", "This mapping name will be the key that is used to save player mapping data for this action. It can be overriden by individual mappings, or specified by the Input Action. These names need to be unique for each player mapping you wish to save.");
+	return Tooltip;
+}
+
+void FPlayerMappableKeyChildSettingsCustomization::OnMappingNameTextCommited(const FText& NewText, ETextCommit::Type InTextCommit)
+{
+	// set the actual property value on the settings object
+	if (InTextCommit == ETextCommit::OnEnter || InTextCommit == ETextCommit::OnUserMovedFocus)
+	{
+		MappingNamePropHandle->SetValue(NewText.ToString());
+	}
+}
+
+bool FPlayerMappableKeyChildSettingsCustomization::OnVerifyMappingName(const FText& InNewText, FText& OutErrorMessage)
+{
+	UPlayerMappableKeySettings* Settings = GetSettingsObject();
+	const FName CurrentName = Settings ? Settings->Name : NAME_None;
+	
+	// Validate that the new name can be used
+	FEnhancedInputPlayerMappableNameValidator NameValidator(CurrentName);
+	EValidatorResult Result = NameValidator.IsValid(InNewText.ToString(), false);
+	OutErrorMessage = FEnhancedInputPlayerMappableNameValidator::GetErrorText(InNewText.ToString(), Result);
+
+	return Result == EValidatorResult::Ok || Result == EValidatorResult::ExistingName;
+}
+
+/////////////////////////////////////////////////////////////////////////////////////
+// FMappingContextKeyContainerCustomization
+
+TSharedRef<IPropertyTypeCustomization> FMappingContextKeyContainerCustomization::MakeInstance()
+{
+	return MakeShareable(new FMappingContextKeyContainerCustomization());
+}
+
+void FMappingContextKeyContainerCustomization::CustomizeHeader(
+	TSharedRef<IPropertyHandle> PropertyHandle,
+	FDetailWidgetRow& HeaderRow,
+	IPropertyTypeCustomizationUtils& CustomizationUtils)
+{
+	MappingsArrayHandle = PropertyHandle->GetChildHandle(GET_MEMBER_NAME_CHECKED(FInputMappingContextMappingData, Mappings));
+	
+	HeaderRow.NameContent()
+	[
+		PropertyHandle->CreatePropertyNameWidget()
+	];
+	
+	HeaderRow.ValueContent()
+	[
+		PropertyHandle->CreatePropertyValueWidget()
+	];	
+}
+
+void FMappingContextKeyContainerCustomization::CustomizeChildren(
+	TSharedRef<IPropertyHandle> PropertyHandle,
+	IDetailChildrenBuilder& ChildBuilder,
+	IPropertyTypeCustomizationUtils& CustomizationUtils)
+{
+	if (!ensure(MappingsArrayHandle.IsValid()))
+	{
+		return;
+	}
+	
+	MappingsArrayHandle->MarkHiddenByCustomization();
+
+	const TSharedRef<FActionMappingsNodeBuilderEx> ActionMappingsBuilder = MakeShareable(new FActionMappingsNodeBuilderEx(MappingsArrayHandle));
+	ChildBuilder.AddCustomBuilder(ActionMappingsBuilder);
+}
+
+#undef LOCTEXT_NAMESPACE
diff --git a/Engine/Plugins/EnhancedInput/Source/InputEditor/Private/EnhancedInputCustomizations.h b/Engine/Plugins/EnhancedInput/Source/InputEditor/Private/EnhancedInputCustomizations.h
new file mode 100644
--- /dev/null
+++ b/Engine/Plugins/EnhancedInput/Source/InputEditor/Private/EnhancedInputCustomizations.h
@@ -0,0 +1,150 @@
+// Copyright Epic Games, Inc. All Rights Reserved.
+
+#pragma once
+
+#include "IDetailCustomization.h"
+#include "IPropertyTypeCustomization.h"
+
+class FKeyStructCustomization;
+class UPlayerMappableKeySettings;
+class IDetailPropertyRow;
+class IDetailLayoutBuilder;
+
+
+class FEnhancedActionMappingCustomization : public IPropertyTypeCustomization
+{
+public:
+	static TSharedRef<IPropertyTypeCustomization> MakeInstance()
+	{
+		return MakeShareable(new FEnhancedActionMappingCustomization());
+	}
+
+	/** IPropertyTypeCustomization interface */
+	virtual void CustomizeHeader(TSharedRef<IPropertyHandle> PropertyHandle, FDetailWidgetRow& HeaderRow, IPropertyTypeCustomizationUtils& CustomizationUtils) override;
+	virtual void CustomizeChildren(TSharedRef<IPropertyHandle> PropertyHandle, IDetailChildrenBuilder& ChildBuilder, IPropertyTypeCustomizationUtils& CustomizationUtils) override;
+
+	/** Searches the trigger array of an input action - returns true if any combo input triggers are present */
+	UE_DEPRECATED(5.8, "The combo trigger is now deprecated")
+	bool DoesTriggerArrayContainCombo() const;
+private:
+
+	void RemoveMappingButton_OnClick() const;
+	void OnTriggersChanged() const;
+	void AddInputActionProperties(TSharedRef<IPropertyHandle> PropertyHandle, IDetailChildrenBuilder& ChildBuilder);
+	void OnInputActionTriggersChanged() const;
+	void OnInputActionModifiersChanged() const;
+
+	TSharedPtr<IPropertyTypeCustomization> KeyStructInstance;
+	TSharedPtr<IPropertyHandle> MappingPropertyHandle;
+	TSharedPtr<FKeyStructCustomization> KeyStructCustomization;
+	IDetailPropertyRow* InputActionTriggersPropertyRow = nullptr;
+	IDetailPropertyRow* InputActionModifiersPropertyRow = nullptr;
+};
+
+/**
+ * Customization for UEnhancedInputDeveloperSettings.
+ *
+ * This will just make the normal details panel for UEnhancedInputDeveloperSettings, and then add all the default settings
+ * of Input Triggers and Input Modifiers by gather all the CDO's for them.
+ */
+class FEnhancedInputDeveloperSettingsCustomization : public IDetailCustomization
+{
+public:
+	//~ IDetailCustomization interface
+	static TSharedRef<IDetailCustomization> MakeInstance();
+
+	virtual ~FEnhancedInputDeveloperSettingsCustomization() override;
+	virtual void CustomizeDetails(IDetailLayoutBuilder& DetailBuilder) override;
+	virtual void CustomizeDetails(const TSharedPtr<IDetailLayoutBuilder>& DetailBuilder) override;
+	//~ End IDetailCustomization interface
+	
+	/** Returns true if something inherits from the UClass passed in. Checks BPs and Native */
+	static bool DoesClassHaveSubtypes(UClass* Class);
+
+private:
+
+	/** Gather all of the native class CDOs for the given class. */
+	static void GatherNativeClassDetailsCDOs(UClass* Class, TArray<UObject*>& CDOs);
+
+	/**
+	 * Called when any Asset is added, removed, or renamed.
+	 *
+	 * This will rebuild the Modifier and Trigger CDO views to make sure that
+	 * any newly added Blueprint
+	 */
+	void RebuildDetailsViewForAsset(const FAssetData& AssetData, const bool bIsAssetBeingRemoved);
+
+	/** Callbacks that are triggered from the Asset Registry. */
+	void OnAssetAdded(const FAssetData& AssetData) { RebuildDetailsViewForAsset(AssetData, false); }
+	void OnAssetRemoved(const FAssetData& AssetData) { RebuildDetailsViewForAsset(AssetData, true); }
+	void OnAssetRenamed(const FAssetData& AssetData, const FString&) { RebuildDetailsViewForAsset(AssetData, false); }
+	
+	/**
+	 * Create a new category on the DetailBuilder and add each object in the given array as an external reference.
+	 * 
+	 * @param DetailBuilder			The details builder that can be used to add categories
+	 * @param CategoryName			The name of the new category to add
+	 * @param ObjectsToCustomize	Array of CDO objects to customize and add as external references
+	 */
+	void CustomizeCDOValues(IDetailLayoutBuilder& DetailBuilder, const FName CategoryName, const TArray<UObject*>& ObjectsToCustomize);
+
+	// Cached details builder so that we can rebuild the details when a new BP asset is added
+	TWeakPtr<IDetailLayoutBuilder> CachedDetailBuilder;
+
+	/**
+	 * Populated by RebuildDetailsViewForAsset so that we can exclude any blueprint classes
+	 * that have been removed from the asset registry but are still loaded in memory.
+	 * This is reset when CustomizeCDOValues is called
+	 */
+	static TSet<FName> ExcludedAssetNames;
+};
+
+/**
+ * Input customization for UPlayerMappableKeySettings to have some name validation
+ * on the custom saved FNames for mappings.
+ */
+class FPlayerMappableKeyChildSettingsCustomization : public IPropertyTypeCustomization
+{
+public:
+	static TSharedRef<IPropertyTypeCustomization> MakeInstance();
+
+	/** IPropertyTypeCustomization interface */
+	virtual void CustomizeHeader(TSharedRef<IPropertyHandle> PropertyHandle, FDetailWidgetRow& HeaderRow, IPropertyTypeCustomizationUtils& CustomizationUtils) override;
+	virtual void CustomizeChildren(TSharedRef<IPropertyHandle> PropertyHandle, IDetailChildrenBuilder& ChildBuilder, IPropertyTypeCustomizationUtils& CustomizationUtils) override;
+
+private:
+	
+	UPlayerMappableKeySettings* GetSettingsObject() const;
+	FName GetCurrentName() const;
+	
+	FText OnGetMappingNameText() const;
+	FText OnGetMappingNameToolTipText() const;
+	void OnMappingNameTextCommited(const FText& NewText, ETextCommit::Type InTextCommit);
+	bool OnVerifyMappingName(const FText& InNewText, FText& OutErrorMessage);
+
+	TSharedPtr<IPropertyHandle> SettingsPropHandle;
+	TSharedPtr<IPropertyHandle> MappingNamePropHandle;
+};
+
+/**
+ * Customization for FInputMappingContextMappingData.
+ */
+class FMappingContextKeyContainerCustomization : public IPropertyTypeCustomization
+{
+public:
+	static TSharedRef<IPropertyTypeCustomization> MakeInstance();
+	
+	virtual void CustomizeHeader(
+		TSharedRef<IPropertyHandle> PropertyHandle,
+		FDetailWidgetRow& HeaderRow,
+		IPropertyTypeCustomizationUtils& CustomizationUtils) override;
+		
+	virtual void CustomizeChildren(
+		TSharedRef<IPropertyHandle> PropertyHandle,
+		IDetailChildrenBuilder& ChildBuilder,
+		IPropertyTypeCustomizationUtils& CustomizationUtils) override;
+
+private:
+	
+	TSharedPtr<IPropertyHandle> MappingsArrayHandle;
+};
\ No newline at end of file
diff --git a/Engine/Plugins/EnhancedInput/Source/InputEditor/Private/EnhancedInputEditorModule.cpp b/Engine/Plugins/EnhancedInput/Source/InputEditor/Private/EnhancedInputEditorModule.cpp
new file mode 100644
--- /dev/null
+++ b/Engine/Plugins/EnhancedInput/Source/InputEditor/Private/EnhancedInputEditorModule.cpp
@@ -0,0 +1,500 @@
+// Copyright Epic Games, Inc. All Rights Reserved.
+
+#include "InputEditorModule.h"
+
+#include "AssetBlueprintGraphActions.h"
+#include "BlueprintEditorModule.h"
+#include "BlueprintNodeTemplateCache.h"
+#include "Brushes/SlateImageBrush.h"
+#include "Framework/Notifications/NotificationManager.h"
+#include "InputMappingContext.h"
+#include "Misc/PackageName.h"
+#include "EnhancedInputCustomizations.h"
+#include "ISettingsModule.h"
+#include "ToolMenuSection.h"
+#include "PropertyEditorModule.h"
+#include "UObject/UObjectIterator.h"
+#include "Widgets/Notifications/SNotificationList.h"
+#include "AssetTypeActions/AssetTypeActions_DataAsset.h"
+#include "EnhancedInputDeveloperSettings.h"
+#include "Styling/SlateStyle.h"
+#include "Styling/SlateStyleMacros.h"
+#include "Styling/SlateStyleRegistry.h"
+#include "ContentBrowserModule.h"
+#include "EdGraphSchema_K2_Actions.h"
+#include "IContentBrowserSingleton.h"
+#include "ClassViewerModule.h"
+#include "ClassViewerFilter.h"
+#include "Kismet2/SClassPickerDialog.h"
+#include "GameFramework/InputSettings.h"
+#include "EnhancedInputComponent.h"
+#include "EnhancedPlayerInput.h"
+#include "Interfaces/IMainFrameModule.h"
+#include "SourceControlHelpers.h"
+#include "HAL/FileManager.h"
+#include "PlayerMappableKeySettings.h"
+#include "IAssetTools.h"
+#include "IAssetTypeActions.h"
+
+#include UE_INLINE_GENERATED_CPP_BY_NAME(InputEditorModule)
+
+#define LOCTEXT_NAMESPACE "InputEditor"
+
+DEFINE_LOG_CATEGORY(LogEnhancedInputEditor);
+
+EAssetTypeCategories::Type FInputEditorModule::InputAssetsCategory;
+
+namespace UE::Input
+{
+	static bool bEnableAutoUpgradeToEnhancedInput = true;
+	static FAutoConsoleVariableRef CVarEnableAutoUpgradeToEnhancedInput(
+		TEXT("EnhancedInput.bEnableAutoUpgrade"),
+		bEnableAutoUpgradeToEnhancedInput,
+		TEXT("Should your project automatically be set to use Enhanced Input if it is currently using the legacy input system?"),
+		ECVF_Default);
+}
+
+IMPLEMENT_MODULE(FInputEditorModule, InputEditor)
+
+class FInputClassParentFilter : public IClassViewerFilter
+{
+public:
+	FInputClassParentFilter()
+		: DisallowedClassFlags(CLASS_Abstract | CLASS_Deprecated | CLASS_NewerVersionExists | CLASS_HideDropDown) {}
+
+	/** All children of these classes will be included unless filtered out by another setting. */
+	TSet<const UClass*> AllowedChildrenOfClasses;
+
+	/** Disallowed class flags. */
+	EClassFlags DisallowedClassFlags;
+
+	virtual bool IsClassAllowed(const FClassViewerInitializationOptions& InInitOptions, const UClass* InClass, TSharedRef<FClassViewerFilterFuncs> InFilterFuncs) override
+	{
+		return !InClass->HasAnyClassFlags(DisallowedClassFlags)
+			&& InFilterFuncs->IfInChildOfClassesSet(AllowedChildrenOfClasses, InClass) != EFilterReturn::Failed;
+	}
+
+	virtual bool IsUnloadedClassAllowed(const FClassViewerInitializationOptions& InInitOptions, const TSharedRef<const IUnloadedBlueprintData> InUnloadedClassData, TSharedRef<FClassViewerFilterFuncs> InFilterFuncs) override
+	{
+		return !InUnloadedClassData->HasAnyClassFlags(DisallowedClassFlags)
+			&& InFilterFuncs->IfInChildOfClassesSet(AllowedChildrenOfClasses, InUnloadedClassData) != EFilterReturn::Failed;
+	}
+};
+
+// Asset factories
+
+// InputContext
+UInputMappingContext_Factory::UInputMappingContext_Factory(const class FObjectInitializer& OBJ) : Super(OBJ) {
+	SupportedClass = UInputMappingContext::StaticClass();
+	bEditAfterNew = true;
+	bCreateNew = true;
+}
+
+bool UInputMappingContext_Factory::ConfigureProperties()
+{
+	if (!FEnhancedInputDeveloperSettingsCustomization::DoesClassHaveSubtypes(UInputMappingContext::StaticClass()))
+	{
+		return true;
+	}
+	
+	// nullptr the InputMappingContextClass so we can check for selection
+	InputMappingContextClass = nullptr;
+
+	// Load the classviewer module to display a class picker
+	FClassViewerModule& ClassViewerModule = FModuleManager::LoadModuleChecked<FClassViewerModule>("ClassViewer");
+
+	// Fill in options
+	FClassViewerInitializationOptions Options;
+	Options.Mode = EClassViewerMode::ClassPicker;
+
+	TSharedPtr<FInputClassParentFilter> Filter = MakeShareable(new FInputClassParentFilter);
+	Filter->AllowedChildrenOfClasses.Add(UInputMappingContext::StaticClass());
+
+	Options.ClassFilters.Add(Filter.ToSharedRef());
+
+	const FText TitleText = LOCTEXT("CreateInputMappingContextOptions", "Pick Class For Input Mapping Context Instance");
+	UClass* ChosenClass = nullptr;
+	const bool bPressedOk = SClassPickerDialog::PickClass(TitleText, Options, ChosenClass, UInputMappingContext::StaticClass());
+
+	if (bPressedOk)
+	{
+		InputMappingContextClass = ChosenClass;
+	}
+
+	return bPressedOk;
+}
+
+UObject* UInputMappingContext_Factory::FactoryCreateNew(UClass* Class, UObject* InParent, FName Name, EObjectFlags Flags, UObject* Context, FFeedbackContext* Warn)
+{
+	UInputMappingContext* IMC = nullptr;
+
+	if (InputMappingContextClass != nullptr)
+	{
+		IMC = NewObject<UInputMappingContext>(InParent, InputMappingContextClass, Name, Flags | RF_Transactional, Context);
+	}
+	else
+	{
+		check(Class->IsChildOf(UInputMappingContext::StaticClass()));
+		IMC = NewObject<UInputMappingContext>(InParent, Class, Name, Flags | RF_Transactional, Context);
+	}
+
+	check(IMC);
+
+	// Populate the IMC with some initial input actions if they were specified. This will be the case if the IMC is being created from the FAssetTypeActions_InputAction
+	for (TWeakObjectPtr<UInputAction> WeakIA : InitialActions)
+	{
+		if (UInputAction* IA = WeakIA.Get())
+		{
+			IMC->MapKey(IA, FKey());
+		}
+	}
+	
+	return IMC;
+}
+
+void UInputMappingContext_Factory::SetInitialActions(TArray<TWeakObjectPtr<UInputAction>> InInitialActions)
+{
+	InitialActions.Empty();
+	InitialActions = InInitialActions;
+}
+
+// InputAction
+UInputAction_Factory::UInputAction_Factory(const class FObjectInitializer& OBJ)
+	: Super(OBJ)
+{
+	SupportedClass = UInputAction::StaticClass();
+	bEditAfterNew = true;
+	bCreateNew = true;
+}
+
+bool UInputAction_Factory::ConfigureProperties()
+{
+	if (!FEnhancedInputDeveloperSettingsCustomization::DoesClassHaveSubtypes(UInputAction::StaticClass()))
+	{
+		return true;
+	}
+	
+	// nullptr the InputActionClass so we can check for selection
+	InputActionClass = nullptr;
+
+	// Load the classviewer module to display a class picker
+	FClassViewerModule& ClassViewerModule = FModuleManager::LoadModuleChecked<FClassViewerModule>("ClassViewer");
+
+	// Fill in options
+	FClassViewerInitializationOptions Options;
+	Options.Mode = EClassViewerMode::ClassPicker;
+
+	TSharedPtr<FInputClassParentFilter> Filter = MakeShareable(new FInputClassParentFilter);
+	Filter->AllowedChildrenOfClasses.Add(UInputAction::StaticClass());
+
+	Options.ClassFilters.Add(Filter.ToSharedRef());
+
+	const FText TitleText = LOCTEXT("CreateInputActionOptions", "Pick Class For Input Action Instance");
+	UClass* ChosenClass = nullptr;
+	const bool bPressedOk = SClassPickerDialog::PickClass(TitleText, Options, ChosenClass, UInputAction::StaticClass());
+
+	if (bPressedOk)
+	{
+		InputActionClass = ChosenClass;
+	}
+
+	return bPressedOk;
+}
+
+UObject* UInputAction_Factory::FactoryCreateNew(UClass* Class, UObject* InParent, FName Name, EObjectFlags Flags, UObject* Context, FFeedbackContext* Warn)
+{	
+	if (InputActionClass != nullptr)
+	{
+		return NewObject<UInputAction>(InParent, InputActionClass, Name, Flags | RF_Transactional, Context);
+	}
+	else
+	{
+		check(Class->IsChildOf(UInputAction::StaticClass()));
+		return NewObject<UInputAction>(InParent, Class, Name, Flags | RF_Transactional, Context);
+	}
+}
+
+// Asset type actions
+
+class FAssetTypeActions_InputContext : public FAssetTypeActions_Base
+{
+public:
+	virtual FText GetName() const override { return NSLOCTEXT("AssetTypeActions", "AssetTypeActions_InputMappingContext", "Input Mapping Context"); }
+	virtual uint32 GetCategories() override { return FInputEditorModule::GetInputAssetsCategory(); }
+	virtual FColor GetTypeColor() const override { return FColor(255, 255, 127); }
+	virtual FText GetAssetDescription(const FAssetData& AssetData) const override { return NSLOCTEXT("AssetTypeActions", "AssetTypeActions_InputContextDesc", "A collection of device input to action mappings."); }
+	virtual UClass* GetSupportedClass() const override { return UInputMappingContext::StaticClass(); }
+};
+
+class FAssetTypeActions_InputAction : public FAssetTypeActions_Base
+{
+public:
+	virtual FText GetName() const override { return NSLOCTEXT("AssetTypeActions", "AssetTypeActions_InputAction", "Input Action"); }
+	virtual uint32 GetCategories() override { return FInputEditorModule::GetInputAssetsCategory(); }
+	virtual FColor GetTypeColor() const override { return FColor(127, 255, 255); }
+	virtual FText GetAssetDescription(const FAssetData& AssetData) const override { return NSLOCTEXT("AssetTypeActions", "AssetTypeActions_InputActionDesc", "Represents an abstract game action that can be mapped to arbitrary hardware input devices."); }
+	virtual UClass* GetSupportedClass() const override { return UInputAction::StaticClass(); }
+	virtual void GetActions(const TArray<UObject*>& InObjects, struct FToolMenuSection& Section) override;
+
+private:
+
+	void GetCreateContextFromActionsMenu(TArray<TWeakObjectPtr<UInputAction>> InActions);
+	
+};
+
+void FAssetTypeActions_InputAction::GetActions(const TArray<UObject*>& InObjects, FToolMenuSection& Section)
+{
+	TArray<TWeakObjectPtr<UInputAction>> InputActions = GetTypedWeakObjectPtrs<UInputAction>(InObjects);
+	
+	 Section.AddMenuEntry(
+	 	"InputAction_CreateContextFromSelection",	
+	 	LOCTEXT("InputAction_CreateContextFromSelection", "Create an Input Mapping Context"),
+	 	LOCTEXT("InputAction_CreateContextFromSelectionTooltip", "Create an Input Mapping Context that is filled with the selected Input Actions"),
+	 	FSlateIcon(FAppStyle::GetAppStyleSetName(), "ClassIcon.InputMappingContext"),
+	 	FUIAction(FExecuteAction::CreateSP(this, &FAssetTypeActions_InputAction::GetCreateContextFromActionsMenu, InputActions))
+	 );
+}
+
+void FAssetTypeActions_InputAction::GetCreateContextFromActionsMenu(TArray<TWeakObjectPtr<UInputAction>> InActions)
+{
+	if (InActions.IsEmpty())
+	{
+		return;
+	}
+
+	static const FString MappingContextPrefix(TEXT("IMC_"));
+	static const FString ActionPrefixToIgnore(TEXT("IA_"));
+	
+	if (UInputAction* FirstIA = InActions[0].Get())
+	{
+		FString EffectiveIMCName = FirstIA->GetName();
+		EffectiveIMCName.RemoveFromStart(ActionPrefixToIgnore);
+		// Have the IMC_ prefix at the front of the new asset
+		EffectiveIMCName.InsertAt(0, MappingContextPrefix);
+		
+		const FString ActionPathName = FirstIA->GetOutermost()->GetPathName();
+    	const FString LongPackagePath = FPackageName::GetLongPackagePath(ActionPathName);
+		const FString NewIMCDefaultPath = LongPackagePath / EffectiveIMCName;
+
+		FAssetToolsModule& AssetToolsModule = FModuleManager::GetModuleChecked<FAssetToolsModule>("AssetTools");
+		FContentBrowserModule& ContentBrowserModule = FModuleManager::LoadModuleChecked<FContentBrowserModule>("ContentBrowser");
+		
+		// Make sure the name is unique
+		FString AssetName;
+		FString DefaultSuffix;
+		FString PackageName;
+		AssetToolsModule.Get().CreateUniqueAssetName(NewIMCDefaultPath, DefaultSuffix, /*out*/ PackageName, /*out*/ AssetName);
+		const FString PackagePath = FPackageName::GetLongPackagePath(PackageName);
+
+		// Create the new IMC
+		UInputMappingContext_Factory* IMCFactory = NewObject<UInputMappingContext_Factory>();
+		IMCFactory->SetInitialActions(InActions);
+		ContentBrowserModule.Get().CreateNewAsset(AssetName, PackagePath, UInputMappingContext::StaticClass(), IMCFactory);
+	}
+}
+
+/** Custom style set for Enhanced Input */
+class FEnhancedInputSlateStyle final : public FSlateStyleSet
+{
+public:
+	FEnhancedInputSlateStyle()
+		: FSlateStyleSet("EnhancedInputEditor")
+	{
+		SetParentStyleName(FAppStyle::GetAppStyleSetName());
+
+		// The icons are located in /Engine/Plugins/EnhancedInput/Content/Editor/Slate/Icons
+		SetContentRoot(FPaths::EnginePluginsDir() / TEXT("EnhancedInput/Content/Editor/Slate"));
+		SetCoreContentRoot(FPaths::EngineContentDir() / TEXT("Slate"));
+
+		// Enhanced Input Editor icons
+		static const FVector2D Icon16 = FVector2D(16.0f, 16.0f);
+		static const FVector2D Icon64 = FVector2D(64.0f, 64.0f);
+
+		Set("EnhancedInputIcon_Small", new IMAGE_BRUSH_SVG("Icons/EnhancedInput_16", Icon16));
+		Set("EnhancedInputIcon_Large", new IMAGE_BRUSH_SVG("Icons/EnhancedInput_64", Icon64));
+		
+		Set("ClassIcon.InputAction", new IMAGE_BRUSH_SVG("Icons/InputAction_16", Icon16));
+		Set("ClassThumbnail.InputAction", new IMAGE_BRUSH_SVG("Icons/InputAction_64", Icon64));
+		
+		Set("ClassIcon.InputMappingContext", new IMAGE_BRUSH_SVG("Icons/InputMappingContext_16", Icon16));
+		Set("ClassThumbnail.InputMappingContext", new IMAGE_BRUSH_SVG("Icons/InputMappingContext_64", Icon64));
+	}
+};
+
+void FInputEditorModule::StartupModule()
+{
+	LLM_SCOPE_BYNAME(TEXT("EnhancedInput/Editor"));
+	
+	// Register customizations
+	FPropertyEditorModule& PropertyModule = FModuleManager::GetModuleChecked<FPropertyEditorModule>("PropertyEditor");
+	PropertyModule.RegisterCustomPropertyTypeLayout(FInputMappingContextMappingData::StaticStruct()->GetFName(), FOnGetPropertyTypeCustomizationInstance::CreateStatic(&FMappingContextKeyContainerCustomization::MakeInstance));
+	PropertyModule.RegisterCustomPropertyTypeLayout("EnhancedActionKeyMapping", FOnGetPropertyTypeCustomizationInstance::CreateStatic(&FEnhancedActionMappingCustomization::MakeInstance));
+	PropertyModule.RegisterCustomClassLayout(UEnhancedInputDeveloperSettings::StaticClass()->GetFName(), FOnGetDetailCustomizationInstance::CreateStatic(&FEnhancedInputDeveloperSettingsCustomization::MakeInstance));
+	PropertyModule.RegisterCustomPropertyTypeLayout(UPlayerMappableKeySettings::StaticClass()->GetFName(), FOnGetPropertyTypeCustomizationInstance::CreateStatic(&FPlayerMappableKeyChildSettingsCustomization::MakeInstance));
+	PropertyModule.NotifyCustomizationModuleChanged();
+
+	// Register input assets
+	IAssetTools& AssetTools = FModuleManager::LoadModuleChecked<FAssetToolsModule>("AssetTools").Get();
+	InputAssetsCategory = AssetTools.RegisterAdvancedAssetCategory(FName(TEXT("Input")), LOCTEXT("InputAssetsCategory", "Input"));
+	{
+		RegisterAssetTypeActions(AssetTools, MakeShareable(new FAssetTypeActions_InputAction));
+		RegisterAssetTypeActions(AssetTools, MakeShareable(new FAssetTypeActions_InputContext));
+	}
+
+	// Make a new style set for Enhanced Input, which will register any custom icons for the types in this plugin
+	StyleSet = MakeShared<FEnhancedInputSlateStyle>();
+	FSlateStyleRegistry::RegisterSlateStyle(*StyleSet.Get());
+
+	// Listen for when the editor is ready and then try to upgrade the input classes. We will send a slate notification
+	// if we upgrade the input classes, which is why we need to wait for the editor to be ready
+    IMainFrameModule& MainFrameModule = FModuleManager::LoadModuleChecked<IMainFrameModule>(TEXT("MainFrame"));
+    MainFrameModule.OnMainFrameCreationFinished().AddRaw(this, &FInputEditorModule::OnMainFrameCreationFinished);
+}
+
+void FInputEditorModule::ShutdownModule()
+{
+	// Unregister input assets
+	if (FAssetToolsModule* AssetToolsModule = FModuleManager::GetModulePtr<FAssetToolsModule>("AssetTools"))
+	{
+		for (TSharedPtr<IAssetTypeActions>& AssetAction : CreatedAssetTypeActions)
+		{
+			AssetToolsModule->Get().UnregisterAssetTypeActions(AssetAction.ToSharedRef());
+		}
+	}
+	CreatedAssetTypeActions.Empty();
+
+	// Unregister input settings
+	if (ISettingsModule* SettingsModule = FModuleManager::GetModulePtr<ISettingsModule>("Settings"))
+	{
+		SettingsModule->UnregisterSettings("Project", "Engine", "Enhanced Input");
+	}
+
+	// Unregister customizations
+	FPropertyEditorModule& PropertyModule = FModuleManager::GetModuleChecked<FPropertyEditorModule>("PropertyEditor");
+	PropertyModule.UnregisterCustomPropertyTypeLayout("InputMappingContextMappingData");
+	PropertyModule.UnregisterCustomPropertyTypeLayout("EnhancedActionKeyMapping");
+	PropertyModule.UnregisterCustomClassLayout("EnhancedInputDeveloperSettings");
+	PropertyModule.UnregisterCustomPropertyTypeLayout("PlayerMappableKeySettings");
+	PropertyModule.NotifyCustomizationModuleChanged();
+
+	// Unregister slate stylings
+	if (StyleSet.IsValid())
+	{
+		FSlateStyleRegistry::UnRegisterSlateStyle(*StyleSet.Get());
+	}
+
+	// Remove any listeners to the editor startup delegate
+	if (FModuleManager::Get().IsModuleLoaded("MainFrame"))
+	{
+		IMainFrameModule::Get().OnMainFrameCreationFinished().RemoveAll(this);
+	}
+}
+
+EAssetTypeCategories::Type FInputEditorModule::GetInputAssetsCategory()
+{
+	return InputAssetsCategory;
+}
+
+void FInputEditorModule::OnMainFrameCreationFinished(TSharedPtr<SWindow> InRootWindow, bool bIsRunningStartupDialog)
+{
+	if (FApp::HasProjectName())
+	{
+		AutoUpgradeDefaultInputClasses();
+	}	
+}
+
+namespace UE::Input
+{
+	static void ShowToast(const FText& ClassBeingChanged, const FString& NewClassName, const FString& OldClassName)
+	{
+		FFormatNamedArguments Args;
+		Args.Add(TEXT("ClassBeingChanged"), ClassBeingChanged);
+		Args.Add(TEXT("NewClass"), FText::FromString(NewClassName));
+		Args.Add(TEXT("OldClass"), FText::FromString(OldClassName));
+
+		const FText Message = FText::Format(LOCTEXT("EnhancedInputUpgradeToast", "Upgrading default {ClassBeingChanged} class from '{OldClass}' to '{NewClass}'"), Args);
+		
+		FNotificationInfo* Info = new FNotificationInfo(Message);
+		Info->ExpireDuration = 5.0f;
+		Info->bFireAndForget = true;
+		FSlateNotificationManager::Get().QueueNotification(Info);
+		
+		UE_LOGF(LogEnhancedInputEditor, Log, "Upgrading Default %ls class from '%ls' to '%ls'", *ClassBeingChanged.ToString(), *OldClassName, *NewClassName);
+	}
+}
+
+void FInputEditorModule::AutoUpgradeDefaultInputClasses()
+{
+	if (!UE::Input::bEnableAutoUpgradeToEnhancedInput)
+	{
+		return;
+	}
+	
+	if (UInputSettings* InputSettings = GetMutableDefault<UInputSettings>())
+	{
+		const UClass* OriginalInputComponentClass = InputSettings->GetDefaultInputComponentClass();
+		const UClass* OriginalPlayerInputClass = InputSettings->GetDefaultPlayerInputClass();
+		bool bNeedsConfigSave = false;
+		
+		if (OriginalInputComponentClass && OriginalInputComponentClass == UInputComponent::StaticClass())
+		{
+			InputSettings->SetDefaultInputComponentClass(UEnhancedInputComponent::StaticClass());
+
+			static const FText ClassName = LOCTEXT("InputComponentClassLabel", "Input Component");
+			UE::Input::ShowToast(ClassName, UEnhancedInputComponent::StaticClass()->GetName(), OriginalInputComponentClass->GetName());
+			bNeedsConfigSave = true;
+		}
+
+		if (OriginalPlayerInputClass && OriginalPlayerInputClass == UPlayerInput::StaticClass())
+		{
+			InputSettings->SetDefaultPlayerInputClass(UEnhancedPlayerInput::StaticClass());
+
+			static const FText ClassName = LOCTEXT("PlayerInputClassLabel", "Player Input");
+			UE::Input::ShowToast(ClassName, UEnhancedPlayerInput::StaticClass()->GetName(), OriginalPlayerInputClass->GetName());
+			bNeedsConfigSave = true;
+		}
+
+		// Make sure that the config file gets updated with these new values
+		if (bNeedsConfigSave)
+		{
+			const FString DefaultConfigFile = InputSettings->GetDefaultConfigFilename();
+			
+			// We can write to the file if it is not read only. If it is read only, then we can write to it if we successfully check it out with source control
+			bool bCanWriteToFile = !IFileManager::Get().IsReadOnly(*DefaultConfigFile) || (USourceControlHelpers::IsEnabled() && USourceControlHelpers::CheckOutFile(DefaultConfigFile));
+			
+			if (bCanWriteToFile)
+			{
+				InputSettings->TryUpdateDefaultConfigFile(DefaultConfigFile, /* bWarnIfFail */ false);
+			}
+		}
+	}
+}
+
+const UPlayerMappableKeySettings* FInputEditorModule::FindMappingByName(const FName InName)
+{
+	// Instanced objects that have been deleted will have an outer of the transient package, and we don't want to include them
+	for (TObjectIterator<UPlayerMappableKeySettings> Itr(RF_ClassDefaultObject | RF_Transient, true, EInternalObjectFlags::Garbage | EInternalObjectFlags::Unreachable); Itr; ++Itr)
+	{
+		const UPlayerMappableKeySettings* Settings = *Itr;
+		if (Settings && Settings->GetOuter() != GetTransientPackage() && Settings->Name == InName)
+		{
+			// What if we had just removed a mapping of this name?
+			return Settings;
+		}
+	}
+	return nullptr;
+}
+
+bool FInputEditorModule::IsMappingNameInUse(const FName InName)
+{
+	// If a mapping with this name exists, then it is in use
+	return FindMappingByName(InName) != nullptr;
+}
+
+void FInputEditorModule::RegisterAssetTypeActions(IAssetTools& AssetTools, TSharedRef<IAssetTypeActions> Action)
+{
+	AssetTools.RegisterAssetTypeActions(Action);
+	CreatedAssetTypeActions.Add(Action);
+}
+
+#undef LOCTEXT_NAMESPACE
\ No newline at end of file
diff --git a/Engine/Plugins/EnhancedInput/Source/InputEditor/Private/InputCustomizations.cpp b/Engine/Plugins/EnhancedInput/Source/InputEditor/Private/InputCustomizations.cpp
deleted file mode 100644
diff --git a/Engine/Plugins/EnhancedInput/Source/InputEditor/Private/InputCustomizations.h b/Engine/Plugins/EnhancedInput/Source/InputEditor/Private/InputCustomizations.h
deleted file mode 100644
diff --git a/Engine/Plugins/EnhancedInput/Source/InputEditor/Private/InputEditorModule.cpp b/Engine/Plugins/EnhancedInput/Source/InputEditor/Private/InputEditorModule.cpp
deleted file mode 100644
