[Enhanced Input] Enable the user settings system by default.

#rnx
#Synced-CL 52004614

diff --git a/Engine/Plugins/EnhancedInput/Source/EnhancedInput/Private/EnhancedInputDeveloperSettings.cpp b/Engine/Plugins/EnhancedInput/Source/EnhancedInput/Private/EnhancedInputDeveloperSettings.cpp
--- a/Engine/Plugins/EnhancedInput/Source/EnhancedInput/Private/EnhancedInputDeveloperSettings.cpp
+++ b/Engine/Plugins/EnhancedInput/Source/EnhancedInput/Private/EnhancedInputDeveloperSettings.cpp
@@ -14,7 +14,7 @@ UEnhancedInputDeveloperSettings::UEnhancedInputDeveloperSettings(const FObjectIn
 	, DefaultPlayerMappableKeyProfileClass(UEnhancedPlayerMappableKeyProfile::StaticClass())
 	, DefaultWorldInputClass(UEnhancedPlayerInput::StaticClass())
 	, bSendTriggeredEventsWhenInputIsFlushed(true)
-	, bEnableUserSettings(false)
+	, bEnableUserSettings(true)
 	, bEnableDefaultMappingContexts(true)
 	, bShouldOnlyTriggerLastActionInChord(true)
 	, bEnableInputModeFiltering(true)
