[Enhanced Input] Fix an operator discrepency in UEnhancedPlayerInput::AddMapping. Use the .Equals function instead of operator== on FEnhancedActionKeyMapping when adding mapping to enhanced input.

We want to compare the triggers/modifiers by their object types, not their exact values. The exact values will never be the same because they are instanced.

#Synced-CL 52059832

diff --git a/Engine/Plugins/EnhancedInput/Source/EnhancedInput/Private/EnhancedPlayerInput.cpp b/Engine/Plugins/EnhancedInput/Source/EnhancedInput/Private/EnhancedPlayerInput.cpp
--- a/Engine/Plugins/EnhancedInput/Source/EnhancedInput/Private/EnhancedPlayerInput.cpp
+++ b/Engine/Plugins/EnhancedInput/Source/EnhancedInput/Private/EnhancedPlayerInput.cpp
@@ -986,14 +986,28 @@ FInputActionValue UEnhancedPlayerInput::GetActionValue(TObjectPtr<const UInputAc
 	return ActionData ? ActionData->GetValue() : FInputActionValue(ForAction->ValueType, FInputActionValue::Axis3D::ZeroVector);
 }
 
-
 int32 UEnhancedPlayerInput::AddMapping(const FEnhancedActionKeyMapping& Mapping)
 {
-	int32 MappingIndex = EnhancedActionMappings.AddUnique(Mapping);
+	// Keep track of what keys are bound in Enhanced Input so that we can query it with IsKeyHandledByAction if desired.
 	++EnhancedKeyBinds.FindOrAdd(Mapping.Key);
+
+	// Flag the key mappings as having changed and requiring a rebuild
 	bKeyMapsBuilt = false;
 
-	return MappingIndex;
+	// Before adding a new mapping, look for a FEnhancedActionKeyMapping that is the same as this newly requested mapping.
+	// We use the ".Equals" function here to compare instead of the raw FEnhancedActionKeyMapping::operator==()
+	// because we want to compare the triggers/modifiers by their object types, not their exact values.
+	// The exact values will never be the same because they are instanced.
+	for (int32 Index = 0; Index < EnhancedActionMappings.Num(); ++Index)
+	{
+		const FEnhancedActionKeyMapping& ExistingMapping = EnhancedActionMappings[Index];
+		if (ExistingMapping.Equals(Mapping))
+		{
+			return Index;
+		}
+	}
+	
+	return EnhancedActionMappings.Add(Mapping);
 }
 
 void UEnhancedPlayerInput::ClearAllMappings()
