[Enhanced Input] Add a C++ automation test for handling multiple IMC's with the same key mapping but with different modifier default values.

#rnx
#Synced-CL 52059832

diff --git a/Engine/Plugins/EnhancedInput/Source/InputEditor/Private/Tests/InputSystemExtendedTest.cpp b/Engine/Plugins/EnhancedInput/Source/InputEditor/Private/Tests/InputSystemExtendedTest.cpp
--- a/Engine/Plugins/EnhancedInput/Source/InputEditor/Private/Tests/InputSystemExtendedTest.cpp
+++ b/Engine/Plugins/EnhancedInput/Source/InputEditor/Private/Tests/InputSystemExtendedTest.cpp
@@ -263,3 +263,84 @@ bool FInputModifierApplicationOrderTest::RunTest(const FString& Parameters)
 
 	return true;
 }
+
+// ============================================================
+// Multi-context priority: modifier values on live key mappings
+// ============================================================
+
+IMPLEMENT_SIMPLE_AUTOMATION_TEST(FInputMappingContextPriorityModifierTest, "Input.System.MappingContextPriorityModifier", BasicSystemExtTestFlags)
+
+bool FInputMappingContextPriorityModifierTest::RunTest(const FString& Parameters)
+{
+	// Two IMCs both map the same action to the same key, each with a Scalar modifier at a
+	// different value. IMC_A is pushed at priority 0, IMC_B at priority 1 (higher).
+	//
+	// After pushing only IMC_A, the single live mapping must carry Scalar = 10.
+	// After also pushing IMC_B, RebuildControlMappings processes higher-priority contexts
+	// first, so IMC_B's mapping appears first in EnhancedActionMappings. FindLiveActionMapping
+	// returns that first entry, which must carry Scalar = 50.
+
+	static const FName ContextA = TEXT("IMC_PriorityTest_A");
+	static const FName ContextB = TEXT("IMC_PriorityTest_B");
+
+	UWorld* World =
+	GIVEN(AnEmptyWorld());
+
+	UControllablePlayer& Data =
+	AND(AControllablePlayer(World));
+	TestTrue(TEXT("Controllable Player is valid"), Data.IsValid());
+
+	AND(AnInputAction(Data, TestAction, EInputActionValueType::Axis1D));
+
+	// -------------------------------------------------------
+	// Step 1: push IMC_A (priority 0) with Scalar(10)
+	// -------------------------------------------------------
+	AND(AnInputContextIsAppliedToAPlayer(Data, ContextA, 0));
+	AND(AnActionIsMappedToAKey(Data, ContextA, TestAction, TestAxis));
+
+	// Pre-configure before applying so DeepCopyPtrArray captures the value.
+	UInputModifierScalar* OrigScalarA = NewObject<UInputModifierScalar>();
+	OrigScalarA->Scalar = FVector(10.f, 10.f, 10.f);
+	AND(AModifierIsAppliedToAnActionMapping(Data, OrigScalarA, ContextA, TestAction, TestAxis));
+
+	// With only IMC_A applied, the single live mapping must have Scalar = 10.
+	FEnhancedActionKeyMapping* LiveMapA = FInputTestHelper::FindLiveActionMapping(Data, TestAction, TestAxis);
+	check(LiveMapA && !LiveMapA->Modifiers.IsEmpty());
+	UInputModifierScalar* LiveScalarA = Cast<UInputModifierScalar>(LiveMapA->Modifiers[0]);
+	check(LiveScalarA);
+	THEN(TestEqual(TEXT("IMC_A only: live modifier scalar is 10"), (float)LiveScalarA->Scalar.X, 10.f));
+
+	// Confirm end-to-end: input 1.0 scaled by 10 -> triggered value 10.
+	WHEN(AKeyIsActuated(Data, TestAxis, 1.0f));
+	AND(InputIsTicked(Data));
+	THEN(PressingKeyTriggersAction(Data, TestAction));
+	AND(TestEqual(TEXT("IMC_A only: triggered value is 10"), FInputTestHelper::GetActionData(Data, TestAction).GetValue().Get<float>(), 10.f));
+
+	// -------------------------------------------------------
+	// Step 2: push IMC_B (priority 1, higher) with Scalar(50)
+	// -------------------------------------------------------
+	AND(AnInputContextIsAppliedToAPlayer(Data, ContextB, 1));
+	AND(AnActionIsMappedToAKey(Data, ContextB, TestAction, TestAxis));
+
+	UInputModifierScalar* OrigScalarB = NewObject<UInputModifierScalar>();
+	OrigScalarB->Scalar = FVector(50.f, 50.f, 50.f);
+	AND(AModifierIsAppliedToAnActionMapping(Data, OrigScalarB, ContextB, TestAction, TestAxis));
+
+	// RebuildControlMappings processes IMC_B (priority 1) before IMC_A (priority 0),
+	// so IMC_B's mapping is first in EnhancedActionMappings. FindLiveActionMapping
+	// returns that entry -- its modifier must now be Scalar = 50.
+	FEnhancedActionKeyMapping* LiveMapB = FInputTestHelper::FindLiveActionMapping(Data, TestAction, TestAxis);
+	check(LiveMapB && !LiveMapB->Modifiers.IsEmpty());
+	UInputModifierScalar* LiveScalarB = Cast<UInputModifierScalar>(LiveMapB->Modifiers[0]);
+	check(LiveScalarB);
+	THEN(TestEqual(TEXT("IMC_A+B: highest-priority live mapping has scalar 50"), (float)LiveScalarB->Scalar.X, 50.f));
+
+	// Confirm end-to-end: with TakeHighestAbsoluteValue, 1.0*50=50 beats 1.0*10=10.
+	// ActionInstanceData was reset by the last rebuild, so state is None -> StartedAndTriggered.
+	WHEN(AKeyIsActuated(Data, TestAxis, 1.0f));
+	AND(InputIsTicked(Data));
+	THEN(PressingKeyTriggersAction(Data, TestAction));
+	AND(TestEqual(TEXT("IMC_A+B: triggered value is 50 (highest wins)"), FInputTestHelper::GetActionData(Data, TestAction).GetValue().Get<float>(), 50.f));
+
+	return true;
+}
