[ContentBrowser] New Add Menu Input Menu

#Synced-CL 52439714

diff --git a/Engine/Plugins/Editor/EngineAssetDefinitions/Source/Private/AssetDefinition_ForceFeedbackAttenuation.h b/Engine/Plugins/Editor/EngineAssetDefinitions/Source/Private/AssetDefinition_ForceFeedbackAttenuation.h
--- a/Engine/Plugins/Editor/EngineAssetDefinitions/Source/Private/AssetDefinition_ForceFeedbackAttenuation.h
+++ b/Engine/Plugins/Editor/EngineAssetDefinitions/Source/Private/AssetDefinition_ForceFeedbackAttenuation.h
@@ -19,7 +19,7 @@ public:
 	virtual TSoftClassPtr<UObject> GetAssetClass() const override { return UForceFeedbackAttenuation::StaticClass(); }
 	virtual TConstArrayView<FAssetCategoryPath> GetAssetCategories() const override
    	{
-   		static const auto Categories = { EAssetCategoryPaths::Input };
+   		static const FAssetCategoryPath Categories[] = { FAssetCategoryPath(EAssetCategoryPaths::Input, NSLOCTEXT("AssetDefinition", "ForceFeedbackAttenuation_SubMenu", "Force Feedback"), ECategoryMenuType::Section) };
    		return Categories;
    	}
 	// UAssetDefinition End
diff --git a/Engine/Plugins/Editor/EngineAssetDefinitions/Source/Private/AssetDefinition_ForceFeedbackEffect.cpp b/Engine/Plugins/Editor/EngineAssetDefinitions/Source/Private/AssetDefinition_ForceFeedbackEffect.cpp
--- a/Engine/Plugins/Editor/EngineAssetDefinitions/Source/Private/AssetDefinition_ForceFeedbackEffect.cpp
+++ b/Engine/Plugins/Editor/EngineAssetDefinitions/Source/Private/AssetDefinition_ForceFeedbackEffect.cpp
@@ -31,7 +31,7 @@ TSoftClassPtr<UObject> UAssetDefinition_ForceFeedbackEffect::GetAssetClass() con
 
 TConstArrayView<FAssetCategoryPath> UAssetDefinition_ForceFeedbackEffect::GetAssetCategories() const
 {
-	static const FAssetCategoryPath Categories[] = { EAssetCategoryPaths::Input };
+	static const FAssetCategoryPath Categories[] = { FAssetCategoryPath(EAssetCategoryPaths::Input, NSLOCTEXT("AssetDefinition", "ForceFeedbackEffect_SubMenu", "Force Feedback"), ECategoryMenuType::Section) };
 	return Categories;
 }
 
diff --git a/Engine/Plugins/Editor/EngineAssetDefinitions/Source/Private/AssetDefinition_HapticFeedbackEffectBuffer.h b/Engine/Plugins/Editor/EngineAssetDefinitions/Source/Private/AssetDefinition_HapticFeedbackEffectBuffer.h
--- a/Engine/Plugins/Editor/EngineAssetDefinitions/Source/Private/AssetDefinition_HapticFeedbackEffectBuffer.h
+++ b/Engine/Plugins/Editor/EngineAssetDefinitions/Source/Private/AssetDefinition_HapticFeedbackEffectBuffer.h
@@ -17,5 +17,10 @@ public:
 	virtual FText GetAssetDisplayName() const override { return NSLOCTEXT("AssetTypeActions", "FAssetTypeActions_HapticFeedbackEffectBuffer", "Haptic Feedback Buffer"); }
 	virtual FLinearColor GetAssetColor() const override { return FLinearColor(FColor(175, 0, 0)); }
 	virtual TSoftClassPtr<UObject> GetAssetClass() const override { return UHapticFeedbackEffect_Buffer::StaticClass(); }
+	virtual TConstArrayView<FAssetCategoryPath> GetAssetCategories() const override
+	{
+		static const FAssetCategoryPath Categories[] = { FAssetCategoryPath(EAssetCategoryPaths::Input, NSLOCTEXT("AssetDefinition", "HapticFeedbackEffectBuffer_SubMenu", "Haptic Feedback"), ECategoryMenuType::Section) };
+		return Categories;
+	}
 	// UAssetDefinition End
 };
diff --git a/Engine/Plugins/Editor/EngineAssetDefinitions/Source/Private/AssetDefinition_HapticFeedbackEffectCurve.h b/Engine/Plugins/Editor/EngineAssetDefinitions/Source/Private/AssetDefinition_HapticFeedbackEffectCurve.h
--- a/Engine/Plugins/Editor/EngineAssetDefinitions/Source/Private/AssetDefinition_HapticFeedbackEffectCurve.h
+++ b/Engine/Plugins/Editor/EngineAssetDefinitions/Source/Private/AssetDefinition_HapticFeedbackEffectCurve.h
@@ -17,5 +17,10 @@ public:
 	virtual FText GetAssetDisplayName() const override { return NSLOCTEXT("AssetTypeActions", "AssetTypeActions_HapticFeedbackEffectCurve", "Haptic Feedback Curve"); }
 	virtual FLinearColor GetAssetColor() const override { return FLinearColor(FColor(175, 0, 0)); }
 	virtual TSoftClassPtr<UObject> GetAssetClass() const override { return UHapticFeedbackEffect_Curve::StaticClass(); }
+	virtual TConstArrayView<FAssetCategoryPath> GetAssetCategories() const override
+	{
+		static const FAssetCategoryPath Categories[] = { FAssetCategoryPath(EAssetCategoryPaths::Input, NSLOCTEXT("AssetDefinition", "HapticFeedbackEffectCurve_SubMenu", "Haptic Feedback"), ECategoryMenuType::Section) };
+		return Categories;
+	}
 	// UAssetDefinition End
 };
diff --git a/Engine/Plugins/Editor/EngineAssetDefinitions/Source/Private/AssetDefinition_HapticFeedbackEffectSoundWave.h b/Engine/Plugins/Editor/EngineAssetDefinitions/Source/Private/AssetDefinition_HapticFeedbackEffectSoundWave.h
--- a/Engine/Plugins/Editor/EngineAssetDefinitions/Source/Private/AssetDefinition_HapticFeedbackEffectSoundWave.h
+++ b/Engine/Plugins/Editor/EngineAssetDefinitions/Source/Private/AssetDefinition_HapticFeedbackEffectSoundWave.h
@@ -17,5 +17,10 @@ public:
 	virtual FText GetAssetDisplayName() const override { return NSLOCTEXT("AssetTypeActions", "FAssetTypeActions_HapticFeedbackEffectSoundWave", "Haptic Feedback Sound Wave"); }
 	virtual FLinearColor GetAssetColor() const override { return FLinearColor(FColor(175, 0, 0)); }
 	virtual TSoftClassPtr<UObject> GetAssetClass() const override { return UHapticFeedbackEffect_SoundWave::StaticClass(); }
+	virtual TConstArrayView<FAssetCategoryPath> GetAssetCategories() const override
+	{
+		static const FAssetCategoryPath Categories[] = { FAssetCategoryPath(EAssetCategoryPaths::Input, NSLOCTEXT("AssetDefinition", "HapticFeedbackEffectSoundWave_SubMenu", "Haptic Feedback"), ECategoryMenuType::Section) };
+		return Categories;
+	}
 	// UAssetDefinition End
 };
diff --git a/Engine/Plugins/Editor/EngineAssetDefinitions/Source/Private/AssetDefinition_TouchInterface.h b/Engine/Plugins/Editor/EngineAssetDefinitions/Source/Private/AssetDefinition_TouchInterface.h
--- a/Engine/Plugins/Editor/EngineAssetDefinitions/Source/Private/AssetDefinition_TouchInterface.h
+++ b/Engine/Plugins/Editor/EngineAssetDefinitions/Source/Private/AssetDefinition_TouchInterface.h
@@ -17,5 +17,10 @@ public:
 	virtual FText GetAssetDisplayName() const override { return NSLOCTEXT("AssetTypeActions", "AssetTypeActions_TouchInterface", "Touch Interface Setup"); }
 	virtual FLinearColor GetAssetColor() const override { return FLinearColor(FColor(175, 0, 128)); }
 	virtual TSoftClassPtr<UObject> GetAssetClass() const override { return UTouchInterface::StaticClass(); }
+	virtual TConstArrayView<FAssetCategoryPath> GetAssetCategories() const override
+	{
+		static const FAssetCategoryPath Categories[] = { FAssetCategoryPath(EAssetCategoryPaths::Input, NSLOCTEXT("AssetDefinition", "TouchInterface_SubMenu", "Touchscreen"), ECategoryMenuType::Section) };
+		return Categories;
+	}
 	// UAssetDefinition End
 };
diff --git a/Engine/Plugins/EnhancedInput/Source/InputEditor/InputEditor.Build.cs b/Engine/Plugins/EnhancedInput/Source/InputEditor/InputEditor.Build.cs
--- a/Engine/Plugins/EnhancedInput/Source/InputEditor/InputEditor.Build.cs
+++ b/Engine/Plugins/EnhancedInput/Source/InputEditor/InputEditor.Build.cs
@@ -14,6 +14,7 @@ namespace UnrealBuildTool.Rules
 
             PrivateDependencyModuleNames.AddRange(
 				new string[] {
+					"AssetDefinition",
 					"ApplicationCore",
 					"BlueprintGraph",
                     "Core",
diff --git a/Engine/Plugins/EnhancedInput/Source/InputEditor/Private/AssetDefinition_InputAction.cpp b/Engine/Plugins/EnhancedInput/Source/InputEditor/Private/AssetDefinition_InputAction.cpp
new file mode 100644
--- /dev/null
+++ b/Engine/Plugins/EnhancedInput/Source/InputEditor/Private/AssetDefinition_InputAction.cpp
@@ -0,0 +1,102 @@
+// Copyright Epic Games, Inc. All Rights Reserved.
+
+#include "AssetDefinition_InputAction.h"
+#include "AssetToolsModule.h"
+#include "ContentBrowserMenuContexts.h"
+#include "ContentBrowserModule.h"
+#include "IContentBrowserSingleton.h"
+#include "InputEditorModule.h"
+#include "InputMappingContext.h"
+#include "ToolMenus.h"
+
+#define LOCTEXT_NAMESPACE "AssetDefinition_InputAction"
+
+namespace MenuExtension::InputAction
+{
+	void GetCreateContextFromActionsMenu(const FToolMenuContext& InContext)
+	{
+		const UContentBrowserAssetContextMenuContext* CBContext = UContentBrowserAssetContextMenuContext::FindContextWithAssets(InContext);
+		if (CBContext == nullptr)
+		{
+			return;
+		}
+
+		TArray<UInputAction*> Actions = CBContext->LoadSelectedObjects<UInputAction>();
+		if (Actions.IsEmpty())
+		{
+			return;
+		}
+
+		static const FString MappingContextPrefix(TEXT("IMC_"));
+		static const FString ActionPrefixToIgnore(TEXT("IA_"));
+	
+		if (UInputAction* FirstIA = Actions[0])
+		{
+			FString EffectiveIMCName = FirstIA->GetName();
+			EffectiveIMCName.RemoveFromStart(ActionPrefixToIgnore);
+			// Have the IMC_ prefix at the front of the new asset
+			EffectiveIMCName.InsertAt(0, MappingContextPrefix);
+		
+			const FString ActionPathName = FirstIA->GetOutermost()->GetPathName();
+			const FString LongPackagePath = FPackageName::GetLongPackagePath(ActionPathName);
+			const FString NewIMCDefaultPath = LongPackagePath / EffectiveIMCName;
+
+			FAssetToolsModule& AssetToolsModule = FModuleManager::GetModuleChecked<FAssetToolsModule>("AssetTools");
+			FContentBrowserModule& ContentBrowserModule = FModuleManager::LoadModuleChecked<FContentBrowserModule>("ContentBrowser");
+		
+			// Make sure the name is unique
+			FString AssetName;
+			FString DefaultSuffix;
+			FString PackageName;
+			AssetToolsModule.Get().CreateUniqueAssetName(NewIMCDefaultPath, DefaultSuffix, /*out*/ PackageName, /*out*/ AssetName);
+			const FString PackagePath = FPackageName::GetLongPackagePath(PackageName);
+
+			// Create the new IMC
+			UInputMappingContext_Factory* IMCFactory = NewObject<UInputMappingContext_Factory>();
+			const TArray<TWeakObjectPtr<UInputAction>> WeakActionsPointers(Actions);
+			IMCFactory->SetInitialActions(WeakActionsPointers);
+			ContentBrowserModule.Get().CreateNewAsset(AssetName, PackagePath, UInputMappingContext::StaticClass(), IMCFactory);
+		}
+	}
+	
+		static FDelayedAutoRegisterHelper DelayedAutoRegister(EDelayedRegisterRunPhase::EndOfEngineInit, []
+	{
+		UToolMenus::RegisterStartupCallback(FSimpleMulticastDelegate::FDelegate::CreateLambda([]()
+		{
+			FToolMenuOwnerScoped OwnerScoped(UE_MODULE_NAME);
+			
+			// Input Action Action Registration
+			{
+				UToolMenu* Menu = UE::ContentBrowser::ExtendToolMenu_AssetContextMenu(UInputAction::StaticClass());
+
+				FToolMenuSection& InputSection = Menu->FindOrAddSection("GetAssetActions");
+
+				InputSection.AddDynamicEntry(NAME_None, FNewToolMenuSectionDelegate::CreateLambda([](FToolMenuSection& InSection)
+					{
+						if (const UContentBrowserAssetContextMenuContext* Context = UContentBrowserAssetContextMenuContext::FindContextWithAssets(InSection))
+						{
+							if (Context->SelectedAssets.Num() > 0)
+							{
+								const FAssetData& AssetData = Context->SelectedAssets[0];
+								if (AssetData.AssetClassPath == UInputAction::StaticClass()->GetClassPathName())
+								{
+									// Create Context From Selection
+									{
+										const TAttribute<FText> Label = LOCTEXT("InputAction_CreateContextFromSelection", "Create an Input Mapping Context");
+										const TAttribute<FText> ToolTip = LOCTEXT("InputAction_CreateContextFromSelectionTooltip", "Create an Input Mapping Context that is filled with the selected Input Actions.");
+										const FSlateIcon Icon = FSlateIcon(FAppStyle::GetAppStyleSetName(), "ClassIcon.InputMappingContext");
+
+										FToolUIAction UIAction;
+										UIAction.ExecuteAction = FToolMenuExecuteAction::CreateStatic(&GetCreateContextFromActionsMenu);
+										InSection.AddMenuEntry("InputAction_CreateContextFromSelection", Label, ToolTip, Icon, UIAction);
+									}
+								}
+							}
+						}
+					}));
+			}
+		}));
+	});
+}
+
+#undef LOCTEXT_NAMESPACE
diff --git a/Engine/Plugins/EnhancedInput/Source/InputEditor/Private/AssetDefinition_InputAction.h b/Engine/Plugins/EnhancedInput/Source/InputEditor/Private/AssetDefinition_InputAction.h
new file mode 100644
--- /dev/null
+++ b/Engine/Plugins/EnhancedInput/Source/InputEditor/Private/AssetDefinition_InputAction.h
@@ -0,0 +1,24 @@
+// Copyright Epic Games, Inc. All Rights Reserved.
+
+#pragma once
+#include "AssetDefinitionDefault.h"
+#include "InputAction.h"
+#include "AssetDefinition_InputAction.generated.h"
+
+UCLASS()
+class UAssetDefinition_InputAction : public UAssetDefinitionDefault
+{
+	GENERATED_BODY()
+public:
+	// UAssetDefinition Begin
+	virtual FText GetAssetDisplayName() const override { return NSLOCTEXT("AssetDefinition", "AssetDefinition_InputAction", "Input Action"); }
+	virtual FLinearColor GetAssetColor() const override { return FColor(127, 255, 255); }
+	virtual FText GetAssetDescription(const FAssetData& AssetData) const override { return NSLOCTEXT("AssetDefinition", "AssetDefinition_InputActionDesc", "Represents an abstract game action that can be mapped to arbitrary hardware input devices."); }
+	virtual TSoftClassPtr<UObject> GetAssetClass() const override { return UInputAction::StaticClass(); }
+	virtual TConstArrayView<FAssetCategoryPath> GetAssetCategories() const override
+	{
+		static const FAssetCategoryPath Categories[] = { FAssetCategoryPath(EAssetCategoryPaths::Input, NSLOCTEXT("AssetDefinition", "AssetDefinition_InputActionSubMenu", "Enhanced Input"), ECategoryMenuType::Section) };
+		return Categories;
+	}
+	// UAssetDefinition End
+};
diff --git a/Engine/Plugins/EnhancedInput/Source/InputEditor/Private/AssetDefinition_InputMappingContext.h b/Engine/Plugins/EnhancedInput/Source/InputEditor/Private/AssetDefinition_InputMappingContext.h
new file mode 100644
--- /dev/null
+++ b/Engine/Plugins/EnhancedInput/Source/InputEditor/Private/AssetDefinition_InputMappingContext.h
@@ -0,0 +1,24 @@
+// Copyright Epic Games, Inc. All Rights Reserved.
+
+#pragma once
+#include "AssetDefinitionDefault.h"
+#include "InputMappingContext.h"
+#include "AssetDefinition_InputMappingContext.generated.h"
+
+UCLASS()
+class UAssetDefinition_InputMappingContext : public UAssetDefinitionDefault
+{
+	GENERATED_BODY()
+public:
+	// UAssetDefinition Begin
+	virtual FText GetAssetDisplayName() const override { return NSLOCTEXT("AssetDefinition", "AssetDefinition_InputMappingContext", "Input Mapping Context"); }
+	virtual FLinearColor GetAssetColor() const override { return FColor(255, 255, 127); }
+	virtual FText GetAssetDescription(const FAssetData& AssetData) const override { return NSLOCTEXT("AssetDefinition", "AssetDefinition_InputMappingContextDesc", "A collection of device input to action mappings."); }
+	virtual TSoftClassPtr<UObject> GetAssetClass() const override { return UInputMappingContext::StaticClass(); }
+	virtual TConstArrayView<FAssetCategoryPath> GetAssetCategories() const override
+	{
+		static const FAssetCategoryPath Categories[] = { FAssetCategoryPath(EAssetCategoryPaths::Input, NSLOCTEXT("AssetDefinition", "AssetDefinition_InputMappingContextSubMenu", "Enhanced Input"), ECategoryMenuType::Section) };
+		return Categories;
+	}
+	// UAssetDefinition End
+};
diff --git a/Engine/Plugins/EnhancedInput/Source/InputEditor/Private/EnhancedInputEditorModule.cpp b/Engine/Plugins/EnhancedInput/Source/InputEditor/Private/EnhancedInputEditorModule.cpp
--- a/Engine/Plugins/EnhancedInput/Source/InputEditor/Private/EnhancedInputEditorModule.cpp
+++ b/Engine/Plugins/EnhancedInput/Source/InputEditor/Private/EnhancedInputEditorModule.cpp
@@ -214,85 +214,6 @@ UObject* UInputAction_Factory::FactoryCreateNew(UClass* Class, UObject* InParent
 	}
 }
 
-// Asset type actions
-
-class FAssetTypeActions_InputContext : public FAssetTypeActions_Base
-{
-public:
-	virtual FText GetName() const override { return NSLOCTEXT("AssetTypeActions", "AssetTypeActions_InputMappingContext", "Input Mapping Context"); }
-	virtual uint32 GetCategories() override { return FInputEditorModule::GetInputAssetsCategory(); }
-	virtual FColor GetTypeColor() const override { return FColor(255, 255, 127); }
-	virtual FText GetAssetDescription(const FAssetData& AssetData) const override { return NSLOCTEXT("AssetTypeActions", "AssetTypeActions_InputContextDesc", "A collection of device input to action mappings."); }
-	virtual UClass* GetSupportedClass() const override { return UInputMappingContext::StaticClass(); }
-};
-
-class FAssetTypeActions_InputAction : public FAssetTypeActions_Base
-{
-public:
-	virtual FText GetName() const override { return NSLOCTEXT("AssetTypeActions", "AssetTypeActions_InputAction", "Input Action"); }
-	virtual uint32 GetCategories() override { return FInputEditorModule::GetInputAssetsCategory(); }
-	virtual FColor GetTypeColor() const override { return FColor(127, 255, 255); }
-	virtual FText GetAssetDescription(const FAssetData& AssetData) const override { return NSLOCTEXT("AssetTypeActions", "AssetTypeActions_InputActionDesc", "Represents an abstract game action that can be mapped to arbitrary hardware input devices."); }
-	virtual UClass* GetSupportedClass() const override { return UInputAction::StaticClass(); }
-	virtual void GetActions(const TArray<UObject*>& InObjects, struct FToolMenuSection& Section) override;
-
-private:
-
-	void GetCreateContextFromActionsMenu(TArray<TWeakObjectPtr<UInputAction>> InActions);
-	
-};
-
-void FAssetTypeActions_InputAction::GetActions(const TArray<UObject*>& InObjects, FToolMenuSection& Section)
-{
-	TArray<TWeakObjectPtr<UInputAction>> InputActions = GetTypedWeakObjectPtrs<UInputAction>(InObjects);
-	
-	 Section.AddMenuEntry(
-	 	"InputAction_CreateContextFromSelection",	
-	 	LOCTEXT("InputAction_CreateContextFromSelection", "Create an Input Mapping Context"),
-	 	LOCTEXT("InputAction_CreateContextFromSelectionTooltip", "Create an Input Mapping Context that is filled with the selected Input Actions"),
-	 	FSlateIcon(FAppStyle::GetAppStyleSetName(), "ClassIcon.InputMappingContext"),
-	 	FUIAction(FExecuteAction::CreateSP(this, &FAssetTypeActions_InputAction::GetCreateContextFromActionsMenu, InputActions))
-	 );
-}
-
-void FAssetTypeActions_InputAction::GetCreateContextFromActionsMenu(TArray<TWeakObjectPtr<UInputAction>> InActions)
-{
-	if (InActions.IsEmpty())
-	{
-		return;
-	}
-
-	static const FString MappingContextPrefix(TEXT("IMC_"));
-	static const FString ActionPrefixToIgnore(TEXT("IA_"));
-	
-	if (UInputAction* FirstIA = InActions[0].Get())
-	{
-		FString EffectiveIMCName = FirstIA->GetName();
-		EffectiveIMCName.RemoveFromStart(ActionPrefixToIgnore);
-		// Have the IMC_ prefix at the front of the new asset
-		EffectiveIMCName.InsertAt(0, MappingContextPrefix);
-		
-		const FString ActionPathName = FirstIA->GetOutermost()->GetPathName();
-    	const FString LongPackagePath = FPackageName::GetLongPackagePath(ActionPathName);
-		const FString NewIMCDefaultPath = LongPackagePath / EffectiveIMCName;
-
-		FAssetToolsModule& AssetToolsModule = FModuleManager::GetModuleChecked<FAssetToolsModule>("AssetTools");
-		FContentBrowserModule& ContentBrowserModule = FModuleManager::LoadModuleChecked<FContentBrowserModule>("ContentBrowser");
-		
-		// Make sure the name is unique
-		FString AssetName;
-		FString DefaultSuffix;
-		FString PackageName;
-		AssetToolsModule.Get().CreateUniqueAssetName(NewIMCDefaultPath, DefaultSuffix, /*out*/ PackageName, /*out*/ AssetName);
-		const FString PackagePath = FPackageName::GetLongPackagePath(PackageName);
-
-		// Create the new IMC
-		UInputMappingContext_Factory* IMCFactory = NewObject<UInputMappingContext_Factory>();
-		IMCFactory->SetInitialActions(InActions);
-		ContentBrowserModule.Get().CreateNewAsset(AssetName, PackagePath, UInputMappingContext::StaticClass(), IMCFactory);
-	}
-}
-
 /** Custom style set for Enhanced Input */
 class FEnhancedInputSlateStyle final : public FSlateStyleSet
 {
@@ -336,10 +257,6 @@ void FInputEditorModule::StartupModule()
 	// Register input assets
 	IAssetTools& AssetTools = FModuleManager::LoadModuleChecked<FAssetToolsModule>("AssetTools").Get();
 	InputAssetsCategory = AssetTools.RegisterAdvancedAssetCategory(FName(TEXT("Input")), LOCTEXT("InputAssetsCategory", "Input"));
-	{
-		RegisterAssetTypeActions(AssetTools, MakeShareable(new FAssetTypeActions_InputAction));
-		RegisterAssetTypeActions(AssetTools, MakeShareable(new FAssetTypeActions_InputContext));
-	}
 
 	// Make a new style set for Enhanced Input, which will register any custom icons for the types in this plugin
 	StyleSet = MakeShared<FEnhancedInputSlateStyle>();
@@ -353,16 +270,6 @@ void FInputEditorModule::StartupModule()
 
 void FInputEditorModule::ShutdownModule()
 {
-	// Unregister input assets
-	if (FAssetToolsModule* AssetToolsModule = FModuleManager::GetModulePtr<FAssetToolsModule>("AssetTools"))
-	{
-		for (TSharedPtr<IAssetTypeActions>& AssetAction : CreatedAssetTypeActions)
-		{
-			AssetToolsModule->Get().UnregisterAssetTypeActions(AssetAction.ToSharedRef());
-		}
-	}
-	CreatedAssetTypeActions.Empty();
-
 	// Unregister input settings
 	if (ISettingsModule* SettingsModule = FModuleManager::GetModulePtr<ISettingsModule>("Settings"))
 	{
@@ -491,10 +398,4 @@ bool FInputEditorModule::IsMappingNameInUse(const FName InName)
 	return FindMappingByName(InName) != nullptr;
 }
 
-void FInputEditorModule::RegisterAssetTypeActions(IAssetTools& AssetTools, TSharedRef<IAssetTypeActions> Action)
-{
-	AssetTools.RegisterAssetTypeActions(Action);
-	CreatedAssetTypeActions.Add(Action);
-}
-
 #undef LOCTEXT_NAMESPACE
\ No newline at end of file
diff --git a/Engine/Plugins/EnhancedInput/Source/InputEditor/Public/InputEditorModule.h b/Engine/Plugins/EnhancedInput/Source/InputEditor/Public/InputEditorModule.h
--- a/Engine/Plugins/EnhancedInput/Source/InputEditor/Public/InputEditorModule.h
+++ b/Engine/Plugins/EnhancedInput/Source/InputEditor/Public/InputEditorModule.h
@@ -40,8 +40,6 @@ public:
 	INPUTEDITOR_API static bool IsMappingNameInUse(const FName InName);
 		
 private:
-	void RegisterAssetTypeActions(IAssetTools& AssetTools, TSharedRef<IAssetTypeActions> Action);
-
 	void OnMainFrameCreationFinished(TSharedPtr<SWindow> InRootWindow, bool bIsRunningStartupDialog);
 	
 	/** Automatically upgrade the current project to use Enhanced Input if it is currently set to the legacy input classes. */
@@ -49,8 +47,6 @@ private:
 
 	static EAssetTypeCategories::Type InputAssetsCategory;
 	
-	TArray<TSharedPtr<IAssetTypeActions>> CreatedAssetTypeActions;
-	
 	TSharedPtr<FSlateStyleSet> StyleSet;
 };
 
