[Enhnaced Input] For boolean input actions mapped to touch Fkeys (FKey::IsTouch), the X,Y values of the keys will be the screen coordinate position. this means that a boolean input action will always evaluate to true, even after the touch has been released. To fix this, have Enhanced Input check the key state's bPressed value instead of the raw X,Y value in UEnhancedPlayerInput::PrepareInputDelegatesForEvaluation. This will ensure that touch FKey values have the expected trigger state flow and are calculated correctly in the IsActuated function of UInputTrigger.

This change is behind the CVar EnhancedInput.CorrectTouchBoolActionKeys, which is defaulted to on, in case any existing content is relying on this broken behavior. This CVar will be removed in a future release.

Add some C++ unit tests for this bug in InputTriggerTest.cpp.

This started as a PR:
--------------------------------
PR #14340: Fix release events for enhanced input screen touches

---------------------------------
#Synced-CL 52624599

diff --git a/Engine/Plugins/EnhancedInput/Source/EnhancedInput/Private/EnhancedPlayerInput.cpp b/Engine/Plugins/EnhancedInput/Source/EnhancedInput/Private/EnhancedPlayerInput.cpp
--- a/Engine/Plugins/EnhancedInput/Source/EnhancedInput/Private/EnhancedPlayerInput.cpp
+++ b/Engine/Plugins/EnhancedInput/Source/EnhancedInput/Private/EnhancedPlayerInput.cpp
@@ -38,6 +38,11 @@ namespace UE::Input
 	static FAutoConsoleVariableRef CVarDetectSubTickTapsForNonAxisInputs(TEXT("EnhancedInput.DetectSubTickTapsForNonAxisInputs"),
 		bDetectSubTickTapsForNonAxisInputs,
 		TEXT("If true, then vast press-release pairs that occur within a tick will trigger for non-axis inputs just as they already do for axis buttons."));
+	
+	static bool bCorrectTouchBoolActionKeys = true;
+	static FAutoConsoleVariableRef CVarCorrectTouchBoolActionKeys(TEXT("EnhancedInput.CorrectTouchBoolActionKeys"),
+		bCorrectTouchBoolActionKeys,
+		TEXT("If true, boolean input actions mapped to touch keys will apply a bug fix to correctly evaluate IsActuated. This will be removed in a future release."));
 
 	static bool IsKeyValidForAnyKey(const FKey& Key)
 	{
@@ -523,6 +528,19 @@ void UEnhancedPlayerInput::PrepareInputDelegatesForEvaluation(const TArray<UInpu
 		FVector RawKeyValue = KeyState ? KeyState->RawValue : FVector::ZeroVector;
 		//UE_CLOGF(RawKeyValue.SizeSquared(), LogEnhancedInput, Warning, "Key %ls - state %ls", *Mapping.Key.GetDisplayName().ToString(), *RawKeyValue.ToString());
 
+		// UPlayerInput::InputTouch stores the screen position in RawValue for all touch events,
+		// including Ended. For Boolean actions, IsActuated() checks RawValue magnitude against the
+		// actuation threshold — screen coordinates always exceed it, so the action appears permanently
+		// actuated and Release/Hold/Tap triggers never work. Use bDown instead for Boolean touch actions.
+		if (UE::Input::bCorrectTouchBoolActionKeys &&
+			KeyState &&
+			Mapping.Action->ValueType == EInputActionValueType::Boolean &&
+			Mapping.Key.IsTouch())
+		{
+			RawKeyValue.X = KeyState->bDown;
+			RawKeyValue.Y = 0.0f;
+		}
+
 		// Should this key be ignored because it was down during a context switch?
 		// If so, check if it is back up, otherwise ignore it.
 		if(Mapping.bShouldBeIgnored && KeyState)
diff --git a/Engine/Plugins/EnhancedInput/Source/InputEditor/Private/Tests/InputTriggerTest.cpp b/Engine/Plugins/EnhancedInput/Source/InputEditor/Private/Tests/InputTriggerTest.cpp
--- a/Engine/Plugins/EnhancedInput/Source/InputEditor/Private/Tests/InputTriggerTest.cpp
+++ b/Engine/Plugins/EnhancedInput/Source/InputEditor/Private/Tests/InputTriggerTest.cpp
@@ -597,6 +597,113 @@ bool FInputTriggerChordsPressedTriggersTest::RunTest(const FString& Parameters)
 	return true;
 }
 
+IMPLEMENT_SIMPLE_AUTOMATION_TEST(FInputTriggerTouchBooleanReleasedTest, "Input.Triggers.Touch.BooleanReleased", BasicTriggerTestFlags)
+
+bool FInputTriggerTouchBooleanReleasedTest::RunTest(const FString& Parameters)
+{
+	// Test that Boolean actions on touch keys correctly fire Released triggers, even when RawValue
+	// contains screen coordinates (which have magnitude >> actuation threshold).
+	// This verifies the fix in EnhancedPlayerInput::PrepareInputDelegatesForEvaluation that
+	// substitutes bDown for screen coordinates on Boolean touch actions.
+
+	UWorld* World =
+	GIVEN(AnEmptyWorld());
+
+	UControllablePlayer& Data =
+	AND(AControllablePlayer(World));
+
+	AND(AnInputContextIsAppliedToAPlayer(Data, TestContext, 0));
+	AND(AnInputAction(Data, TestAction, EInputActionValueType::Boolean));
+	AND(AnActionIsMappedToAKey(Data, TestContext, TestAction, EKeys::TouchKeys[0]));
+	AND(ATriggerIsAppliedToAnAction(Data, NewObject<UInputTriggerReleased>(), TestAction));
+
+	TRIGGER_SUBTEST("1 - Press touch, should start Released trigger (Ongoing)")
+	{
+		WHEN(AKeyIsActuated(Data, EKeys::TouchKeys[0]));
+		// Simulate touch screen coordinates (critical for reproducing the bug)
+		FKeyState* TouchKeyState = Data.PlayerInput->GetKeyState(EKeys::TouchKeys[0]);
+		if (TouchKeyState)
+		{
+			TouchKeyState->RawValueAccumulator = FVector(500.0f, 300.0f, 0.0f);
+		}
+		AND(InputIsTicked(Data));
+		THEN(PressingKeyTriggersStarted(Data, TestAction));
+		ANDALSO(!PressingKeyTriggersAction(Data, TestAction));
+	}
+
+	TRIGGER_SUBTEST("2 - Hold touch, should remain Ongoing")
+	{
+		WHEN(AKeyIsActuated(Data, EKeys::TouchKeys[0]));
+		// Simulate finger movement (different screen coordinates)
+		FKeyState* TouchKeyState = Data.PlayerInput->GetKeyState(EKeys::TouchKeys[0]);
+		if (TouchKeyState)
+		{
+			TouchKeyState->RawValueAccumulator = FVector(502.0f, 301.0f, 0.0f);
+		}
+		AND(InputIsTicked(Data));
+		THEN(HoldingKeyTriggersOngoing(Data, TestAction));
+	}
+
+	TRIGGER_SUBTEST("3 - Release touch, Released trigger should fire (MAIN TEST)")
+	{
+		WHEN(AKeyIsReleased(Data, EKeys::TouchKeys[0]));
+		// Simulate touch screen coordinates persisting on release (like InputTouch does)
+		// This is the critical reproduction of the real-world bug scenario
+		FKeyState* TouchKeyState = Data.PlayerInput->GetKeyState(EKeys::TouchKeys[0]);
+		if (TouchKeyState)
+		{
+			TouchKeyState->RawValueAccumulator = FVector(502.0f, 301.0f, 0.0f);
+		}
+		AND(InputIsTicked(Data));
+		// Without the fix, IsActuated(502, 301, 0) = true (magnitude 502 >> 0.5 threshold)
+		// so Released trigger never fires. With the fix, bDown=0, IsActuated(0)=false, fires correctly.
+		THEN(ReleasingKeyTriggersAction(Data, TestAction));
+	}
+
+	return true;
+}
+
+IMPLEMENT_SIMPLE_AUTOMATION_TEST(FInputTriggerKeyboardBooleanReleasedTest, "Input.Triggers.Keyboard.BooleanReleased", BasicTriggerTestFlags)
+
+bool FInputTriggerKeyboardBooleanReleasedTest::RunTest(const FString& Parameters)
+{
+	// Regression test: verify that keyboard keys with Boolean + Released trigger still work
+	// correctly. The touch coordinate fix (IsTouch() guard) must not interfere with normal keys.
+
+	UWorld* World =
+	GIVEN(AnEmptyWorld());
+
+	UControllablePlayer& Data =
+	AND(AControllablePlayer(World));
+
+	AND(AnInputContextIsAppliedToAPlayer(Data, TestContext, 0));
+	AND(AnInputAction(Data, TestAction, EInputActionValueType::Boolean));
+	AND(AnActionIsMappedToAKey(Data, TestContext, TestAction, EKeys::SpaceBar));
+	AND(ATriggerIsAppliedToAnAction(Data, NewObject<UInputTriggerReleased>(), TestAction));
+
+	TRIGGER_SUBTEST("1 - Press space, should start Released trigger (Ongoing)")
+	{
+		WHEN(AKeyIsActuated(Data, EKeys::SpaceBar));
+		AND(InputIsTicked(Data));
+		THEN(PressingKeyTriggersStarted(Data, TestAction));
+	}
+
+	TRIGGER_SUBTEST("2 - Hold space, should remain Ongoing")
+	{
+		WHEN(InputIsTicked(Data));
+		THEN(HoldingKeyTriggersOngoing(Data, TestAction));
+	}
+
+	TRIGGER_SUBTEST("3 - Release space, Released trigger should fire")
+	{
+		WHEN(AKeyIsReleased(Data, EKeys::SpaceBar));
+		AND(InputIsTicked(Data));
+		THEN(ReleasingKeyTriggersAction(Data, TestAction));
+	}
+
+	return true;
+}
+
 // TODO: Action level triggers (simple repeat of device level tests)
 // TODO: Variable frame delta tests
 // TODO: ActionEventData tests (timing, summed values, etc)
