[Enhanced Input] Fix a bug where key mappings would not correclty be ignored if the input was just flushed (for example, by changing input modes between Game and UI for the player controller)

#Synced-CL 52624599

diff --git a/Engine/Plugins/EnhancedInput/Source/EnhancedInput/Private/EnhancedPlayerInput.cpp b/Engine/Plugins/EnhancedInput/Source/EnhancedInput/Private/EnhancedPlayerInput.cpp
--- a/Engine/Plugins/EnhancedInput/Source/EnhancedInput/Private/EnhancedPlayerInput.cpp
+++ b/Engine/Plugins/EnhancedInput/Source/EnhancedInput/Private/EnhancedPlayerInput.cpp
@@ -545,7 +545,7 @@ void UEnhancedPlayerInput::PrepareInputDelegatesForEvaluation(const TArray<UInpu
 		// If so, check if it is back up, otherwise ignore it.
 		if(Mapping.bShouldBeIgnored && KeyState)
 		{
-			if(KeyState->bDown)
+			if(KeyState->bDown || KeyState->bWasJustFlushed)
 			{
 				continue;				
 			}
